Packet Handler Framework which will use packet handler components to transform incoming and outgoing packets. Currently turned off by default.

diff --git a/Engine/Source/Runtime/Engine/Classes/Engine/NetConnection.h b/Engine/Source/Runtime/Engine/Classes/Engine/NetConnection.h
--- a/Engine/Source/Runtime/Engine/Classes/Engine/NetConnection.h
+++ b/Engine/Source/Runtime/Engine/Classes/Engine/NetConnection.h
@@ -10,6 +10,7 @@
 #include "Engine/Channel.h"
 #include "Engine/Player.h"
 #include "Engine/NetDriver.h"
+#include "Runtime/PacketHandlers/PacketHandler/Public/PacketHandler.h"
 
 #include "NetConnection.generated.h"
 
@@ -204,6 +205,9 @@ public:
 	
 	uint32 bPendingDestroy:1;    // when true, playercontroller is being destroyed
 
+	// Packet Handler
+	TUniquePtr<PacketHandler> Handler;
+
 	/** Whether this channel needs to byte swap all data or not */
 	bool			bNeedsByteSwapping;
 	/** Net id of remote player on this connection. Only valid on client connections. */
diff --git a/Engine/Source/Runtime/Engine/Engine.Build.cs b/Engine/Source/Runtime/Engine/Engine.Build.cs
--- a/Engine/Source/Runtime/Engine/Engine.Build.cs
+++ b/Engine/Source/Runtime/Engine/Engine.Build.cs
@@ -9,7 +9,7 @@ public class Engine : ModuleRules
 	{
 		SharedPCHHeaderFile = "Runtime/Engine/Public/Engine.h";
 
-		PublicIncludePathModuleNames.AddRange(new string[] { "Renderer" });
+		PublicIncludePathModuleNames.AddRange(new string[] { "Renderer", "PacketHandler" });
 
 		PrivateIncludePaths.AddRange(
 			new string[] {
@@ -69,7 +69,8 @@ public class Engine : ModuleRules
                 "UMG",
 				"Projects",
 				"Niagara",
-                "Internationalization"
+                "Internationalization",
+                "PacketHandler"
 			}
         );
 
diff --git a/Engine/Source/Runtime/Engine/Private/NetConnection.cpp b/Engine/Source/Runtime/Engine/Private/NetConnection.cpp
--- a/Engine/Source/Runtime/Engine/Private/NetConnection.cpp
+++ b/Engine/Source/Runtime/Engine/Private/NetConnection.cpp
@@ -12,6 +12,7 @@
 #include "DataChannel.h"
 #include "Engine/PackageMapClient.h"
 #include "GameFramework/GameMode.h"
+#include "Runtime/PacketHandlers/PacketHandler/Public/PacketHandler.h"
 
 #if WITH_EDITOR
 #include "UnrealEd.h"
@@ -455,6 +456,24 @@ void UNetConnection::ReceivedRawPacket( void* InData, int32 Count )
 {
 	uint8* Data = (uint8*)InData;
 
+	// UnProcess the packet
+	if(Handler.IsValid())
+	{
+		const ProcessedPacket UnProcessedPacket = Handler->Incoming(Data, Count);
+
+		Count = UnProcessedPacket.Count;
+
+		if (Count > 0)
+		{
+			Data = UnProcessedPacket.Data;
+		}
+		// This packed has been consumed
+		else
+		{
+			return;
+		}
+	}
+
 	// Handle an incoming raw packet from the driver.
 	UE_LOG(LogNetTraffic, Verbose, TEXT("%6.3f: Received %i"), FPlatformTime::Seconds() - GStartTime, Count );
 	int32 PacketBytes = Count + PacketOverhead;
@@ -1506,6 +1525,21 @@ void UNetConnection::Tick()
 		FlushNet();
 	}
 
+	// Tick Handler
+	if(Handler.IsValid())
+	{
+		Handler->Tick(FrameTime);
+		BufferedPacket* QueuedPacket = Handler->GetQueuedPacket();
+
+		/* Send all queued packets */
+		while(QueuedPacket != nullptr)
+		{
+			LowLevelSend(QueuedPacket->Data, QueuedPacket->BytesCount);
+			delete QueuedPacket;
+			QueuedPacket = Handler->GetQueuedPacket();
+		}
+	}
+
 	// Update queued byte count.
 	// this should be at the end so that the cap is applied *after* sending (and adjusting QueuedBytes for) any remaining data for this tick
 	float DeltaBytes = CurrentNetSpeed * DeltaTime;
diff --git a/Engine/Source/Runtime/Online/OnlineSubsystemUtils/OnlineSubsystemUtils.Build.cs b/Engine/Source/Runtime/Online/OnlineSubsystemUtils/OnlineSubsystemUtils.Build.cs
--- a/Engine/Source/Runtime/Online/OnlineSubsystemUtils/OnlineSubsystemUtils.Build.cs
+++ b/Engine/Source/Runtime/Online/OnlineSubsystemUtils/OnlineSubsystemUtils.Build.cs
@@ -19,7 +19,8 @@ public class OnlineSubsystemUtils : ModuleRules
                 "ImageCore",
 				"OnlineSubsystem",
 				"Sockets",
-				"Voice"
+				"Voice",
+                "PacketHandler"
 			}
 		);
 	}
diff --git a/Engine/Source/Runtime/Online/OnlineSubsystemUtils/Private/IpConnection.cpp b/Engine/Source/Runtime/Online/OnlineSubsystemUtils/Private/IpConnection.cpp
--- a/Engine/Source/Runtime/Online/OnlineSubsystemUtils/Private/IpConnection.cpp
+++ b/Engine/Source/Runtime/Online/OnlineSubsystemUtils/Private/IpConnection.cpp
@@ -13,6 +13,7 @@ Notes:
 #include "Sockets.h"
 #include "Net/NetworkProfiler.h"
 #include "Net/DataChannel.h"
+#include "Runtime/PacketHandlers/PacketHandler/Public/PacketHandler.h"
 
 /*-----------------------------------------------------------------------------
 	Declarations.
@@ -41,6 +42,14 @@ void UIpConnection::InitBase(UNetDriver* InDriver, class FSocket* InSocket, cons
 
 	Socket = InSocket;
 	ResolveInfo = NULL;
+
+	// Reset Handler
+	Handler.Reset(nullptr);
+	if(Handler.IsValid())
+	{
+		Handler::Mode Mode = Driver->ServerConnection != nullptr ? Handler::Mode::Client : Handler::Mode::Server;
+		Handler->Initialize(Mode);
+	}
 }
 
 void UIpConnection::InitLocalConnection(UNetDriver* InDriver, class FSocket* InSocket, const FURL& InURL, EConnectionState InState, int32 InMaxPacket, int32 InPacketOverhead)
@@ -99,6 +108,8 @@ void UIpConnection::InitRemoteConnection(UNetDriver* InDriver, class FSocket* In
 
 void UIpConnection::LowLevelSend( void* Data, int32 Count )
 {
+	const uint8* DataToSend = reinterpret_cast<uint8*>(Data);
+
 	if( ResolveInfo )
 	{
 		// If destination address isn't resolved yet, send nowhere.
@@ -127,10 +138,19 @@ void UIpConnection::LowLevelSend( void* Data, int32 Count )
 			ResolveInfo = NULL;
 		}
 	}
+	 
+	// Process any packet modifiers
+	if(Handler.IsValid())
+	{
+		const ProcessedPacket ProcessedData = Handler->Outgoing(reinterpret_cast<uint8*>(Data), Count);
+		DataToSend = ProcessedData.Data;
+		Count = ProcessedData.Count;
+	}
+
 	// Send to remote.
 	int32 BytesSent = 0;
 	CLOCK_CYCLES(Driver->SendCycles);
-	Socket->SendTo((uint8*)Data, Count, BytesSent, *RemoteAddr);
+	Socket->SendTo(DataToSend, Count, BytesSent, *RemoteAddr);
 	UNCLOCK_CYCLES(Driver->SendCycles);
 	NETWORK_PROFILER(GNetworkProfiler.FlushOutgoingBunches(this));
 	NETWORK_PROFILER(GNetworkProfiler.TrackSocketSendTo(Socket->GetDescription(),Data,BytesSent,NumPacketIdBits,NumBunchBits,NumAckBits,NumPaddingBits,*RemoteAddr));
diff --git a/Engine/Source/Runtime/PacketHandlers/PacketHandler/PacketHandler.Build.cs b/Engine/Source/Runtime/PacketHandlers/PacketHandler/PacketHandler.Build.cs
new file mode 100644
--- /dev/null
+++ b/Engine/Source/Runtime/PacketHandlers/PacketHandler/PacketHandler.Build.cs
@@ -0,0 +1,21 @@
+// Copyright 1998-2015 Epic Games, Inc. All Rights Reserved.
+
+using UnrealBuildTool;
+using System.IO;
+
+public class PacketHandler : ModuleRules
+{
+    public PacketHandler(TargetInfo Target)
+    {
+        PublicDependencyModuleNames.AddRange(
+            new string[] {
+				"Core",
+                "ReliabilityHandlerComponent",
+            }
+        );
+
+        CircularlyReferencedDependentModules.Add("ReliabilityHandlerComponent");
+
+        //DynamicallyLoadedModuleNames.Add("EncryptionHandlerComponent");
+    }
+}
diff --git a/Engine/Source/Runtime/PacketHandlers/PacketHandler/Private/PacketHandler.cpp b/Engine/Source/Runtime/PacketHandlers/PacketHandler/Private/PacketHandler.cpp
new file mode 100644
--- /dev/null
+++ b/Engine/Source/Runtime/PacketHandlers/PacketHandler/Private/PacketHandler.cpp
@@ -0,0 +1,314 @@
+#include "PacketHandler.h"
+#include "ReliabilityHandlerComponent.h"
+
+IMPLEMENT_MODULE(FPacketHandlerModuleInterface, PacketHandler);
+
+DEFINE_LOG_CATEGORY(PacketHandlerLog);
+
+// BUFFERED PACKET
+BufferedPacket::BufferedPacket()
+: Data(nullptr)
+, BytesCount(0)
+, SendTime(0.f)
+, Id(0)
+{
+}
+
+BufferedPacket::~BufferedPacket()
+{
+	if(Data != nullptr)
+	{
+		delete Data;
+	}
+}
+
+// PACKET HANDLER
+PacketHandler::PacketHandler()
+: Time(0.f)
+, State(Handler::State::Uninitialized)
+, ReliabilityComponent(nullptr)
+{
+	OutgoingPacket.SetAllowResize(true);
+	OutgoingPacket.AllowAppend(true);
+}
+
+void PacketHandler::Tick(float DeltaTime)
+{
+	Time += DeltaTime;
+
+	switch (State)
+	{
+		case Handler::State::InitializingProcessors:
+		{
+			break;
+		}
+		case Handler::State::Initialized:
+		{
+			break;
+		}
+	}
+
+	for (int32 i = 0; i < HandlerComponents.Num(); ++i)
+	{
+		HandlerComponents[i]->Tick(DeltaTime);
+	}
+}
+
+void PacketHandler::Initialize(Handler::Mode InMode)
+{
+	Mode = InMode;
+
+	FString Components;
+	GConfig->GetString(TEXT("PacketHandlerComponents"), TEXT("Components"), Components, GEngineIni);
+
+	TArray<FString> ComponentsArray;
+	Components.ParseIntoArray(ComponentsArray, TEXT(","), true);
+
+	for (int32 i = 0; i < ComponentsArray.Num(); i++)
+	{
+		if(ComponentsArray[i] == TEXT("ReliabilityHandlerComponent"))
+		{
+			continue;
+		}
+
+		TSharedPtr<IModuleInterface> Interface = FModuleManager::Get().LoadModule(FName(*ComponentsArray[i]));
+
+		if(Interface.IsValid())
+		{
+			TSharedPtr<FPacketHandlerModuleInterface> PacketHandlerInterface(static_cast<FPacketHandlerModuleInterface*>(&(*Interface)));
+			Add(PacketHandlerInterface->CreateComponentInstance());
+		}
+	}
+}
+
+void PacketHandler::Add(HandlerComponent* NewHandler)
+{
+	if (State != Handler::State::Uninitialized)
+	{
+		LowLevelFatalError(TEXT("Handler added during runtime."));
+	}
+
+	HandlerComponents.Add(NewHandler);
+	NewHandler->Handler = this;
+	NewHandler->Initialize();
+}
+
+const ProcessedPacket PacketHandler::Outgoing(uint8* Packet, int32 Count)
+{
+	OutgoingPacket.Reset();
+
+	switch (State)
+	{
+		case Handler::State::Uninitialized:
+		{
+			if(ReliabilityComponent == nullptr)
+			{
+				ReliabilityComponent = new ReliabilityHandlerComponent;
+				Add(ReliabilityComponent);
+			}
+
+			// Have handlers to initialize other than reliability
+			if (HandlerComponents.Num() > 1)
+			{
+				State = Handler::State::InitializingProcessors;
+			} 
+			else
+			{
+				State = Handler::State::Initialized;
+				break;
+			}
+		}
+		case Handler::State::InitializingProcessors:
+		{
+			// Buffer any packets being sent from game code
+			// until processors are initialized
+			if (Count > 0)
+			{
+				BufferedPacket* NewBufferedPacket = new BufferedPacket;
+				BufferedPackets.Add(NewBufferedPacket);
+
+				uint8* Data = new uint8[Count];
+				memcpy(Data, Packet, Count);
+				NewBufferedPacket->Data = Data;
+				NewBufferedPacket->BytesCount = Count;
+
+				Count = 0;
+				Packet = nullptr;
+			}
+
+			break;
+		}
+		case Handler::State::Initialized:
+		{
+			OutgoingPacket.Serialize(Packet, Count);
+			break;
+		}
+		default:
+		{
+			break;
+		}
+	}
+
+	// Queue packet for resending before handling
+	if(Count > 0 && ReliabilityComponent != nullptr)
+	{
+		ReliabilityComponent->QueuePacketForResending(Packet, Count);
+	}
+
+	// Handle
+	for (int32 i = 0; i < HandlerComponents.Num(); ++i)
+	{
+		if (HandlerComponents[i]->IsActive())
+		{
+			HandlerComponents[i]->Outgoing(OutgoingPacket);
+		}
+	}
+
+	return ProcessedPacket(OutgoingPacket.GetData(), OutgoingPacket.GetNumBytes());
+}
+
+const ProcessedPacket PacketHandler::Incoming(uint8* Packet, int32 Count)
+{
+	FBitReader ProcessedPacketReader(Packet, Count * 8);
+
+	switch (State)
+	{
+		case Handler::State::Uninitialized:
+		{
+			if (ReliabilityComponent == nullptr)
+			{
+				ReliabilityComponent = new ReliabilityHandlerComponent;
+				Add(ReliabilityComponent);
+			}
+
+			if (HandlerComponents.Num() > 1)
+			{
+				State = Handler::State::InitializingProcessors;
+			}
+			else
+			{
+				State = Handler::State::Initialized;
+			}
+		}
+		case Handler::State::InitializingProcessors:
+		{
+			break;
+		}
+		case Handler::State::Initialized:
+		{
+			break;
+		}
+		default:
+		{
+			break;
+		}
+	}
+
+	// Handle
+	for (int32 i = HandlerComponents.Num() - 1; i >= 0; --i)
+	{
+		if (HandlerComponents[i]->IsActive() && ProcessedPacketReader.GetNumBytes() > 0)
+		{
+			HandlerComponents[i]->Incoming(ProcessedPacketReader);
+		}
+	}
+
+	IncomingPacket = ProcessedPacketReader;
+
+	return ProcessedPacket(IncomingPacket.GetData(), IncomingPacket.GetBytesLeft());
+}
+
+void PacketHandler::SetState(Handler::State InState)
+{
+	if (InState == State)
+	{
+		LowLevelFatalError(TEXT("Set new Packet Processor State to the state it is currently in."));
+	} 
+	else
+	{
+		State = InState;
+	}
+}
+
+void PacketHandler::HandlerInitialized()
+{
+	// If any buffered packets, add to queue
+	for (int32 i = 0; i < BufferedPackets.Num(); ++i)
+	{
+		QueuedPackets.Enqueue(BufferedPackets[i]);
+		BufferedPackets[i] = nullptr;
+	}
+
+	SetState(Handler::State::Initialized);
+}
+
+void PacketHandler::HandlerComponentInitialized()
+{
+	// Check if all handlers are initialized
+	if (State != Handler::State::Initialized)
+	{
+		for (int32 i = 0; i < HandlerComponents.Num(); ++i)
+		{
+			if (HandlerComponents[i]->IsInitialized() == false)
+			{
+				return;
+			}
+		}
+
+		HandlerInitialized();
+	}
+}
+
+BufferedPacket* PacketHandler::GetQueuedPacket()
+{
+	BufferedPacket* QueuedPacket = nullptr;
+	QueuedPackets.Dequeue(QueuedPacket);
+	return QueuedPacket;
+}
+
+void PacketHandler::QueuePacketForSending(BufferedPacket* PacketToQueue)
+{
+
+}
+
+// HANDLER COMPONENT
+HandlerComponent::HandlerComponent()
+: bActive(false)
+, State(Handler::Component::State::UnInitialized)
+, bInitialized(false)
+{
+
+}
+
+bool HandlerComponent::IsActive() const
+{
+	return bActive;
+}
+
+void HandlerComponent::SetActive(bool Active)
+{
+	bActive = Active;
+}
+
+void HandlerComponent::SetState(Handler::Component::State InState)
+{
+	State = InState;
+}
+
+void HandlerComponent::Initialized()
+{
+	bInitialized = true;
+	Handler->HandlerComponentInitialized();
+}
+
+bool HandlerComponent::IsInitialized() const
+{
+	return bInitialized;
+}
+
+// PROCESSED PACKET
+ProcessedPacket::ProcessedPacket(uint8* InData, int32 InCount)
+{
+	Data = InData;
+	Count = InCount;
+}
\ No newline at end of file
diff --git a/Engine/Source/Runtime/PacketHandlers/PacketHandler/Public/PacketHandler.h b/Engine/Source/Runtime/PacketHandlers/PacketHandler/Public/PacketHandler.h
new file mode 100644
--- /dev/null
+++ b/Engine/Source/Runtime/PacketHandlers/PacketHandler/Public/PacketHandler.h
@@ -0,0 +1,212 @@
+// Copyright 1998-2015 Epic Games, Inc. All Rights Reserved.
+
+#pragma once
+
+#include "ModuleManager.h"
+#include "Core.h"
+
+DECLARE_LOG_CATEGORY_EXTERN(PacketHandlerLog, Log, All);
+
+class HandlerComponent;
+class ReliabilityHandlerComponent;
+
+/* Enums related to the Packet Handler */
+namespace Handler
+{
+	/* State of Packet Handler */
+	enum State
+	{
+		Uninitialized,
+		InitializingProcessors,
+		Initialized
+	};
+
+	/* Mode of Packet Handler */
+	enum Mode
+	{
+		Client,
+		Server
+	};
+
+	namespace Component
+	{
+		/* Packet Handler State */
+		enum State
+		{
+			UnInitialized,
+			InitializedOnLocal, // Initialized on local instance
+			InitializeOnRemote, // Initialized on remote instance, not on local instance
+			Initialized         // Initialized on both local and remote instances
+		};
+	}
+}
+
+/* The result of calling Incoming and Outgoing in the PacketHandler */
+struct PACKETHANDLER_API ProcessedPacket
+{
+	ProcessedPacket(uint8* Data, int32 Count);
+	uint8* Data;
+	int32 Count;
+};
+
+/*
+* Packet Handler will buffer packets, this struct is used to buffer such packets
+* while handler components are initialized
+*/
+struct PACKETHANDLER_API BufferedPacket
+{
+	BufferedPacket();
+	~BufferedPacket();
+
+	uint8* Data;
+	uint32 BytesCount;
+	float SendTime;
+	uint32 Id;
+};
+
+/*
+* This class maintains an array of all Packet Handler Components and
+* forwards incoming and outgoing packets the each component
+*/
+class PACKETHANDLER_API PacketHandler
+{
+public:
+	/* Default initialization of data */
+	PacketHandler();
+
+	/* Handles initialization of manager */
+	virtual void Initialize(Handler::Mode Mode);
+
+	/* Tick per frame */
+	virtual void Tick(float DeltaTime);
+
+	/* For adding handler components */
+	virtual void Add(HandlerComponent* NewHandler);
+
+	/* Handles any incoming packets */
+	virtual const ProcessedPacket Incoming(uint8* Packet, int32 Count);
+
+	/* Handles any outgoing packets */
+	virtual const ProcessedPacket Outgoing(uint8* Packet, int32 Count);
+
+	/* Indicates that a handler was initialized */
+	void HandlerComponentInitialized();
+
+	/* Queue's a packet to be sent when the handler is ticked */
+	void QueuePacketForSending(BufferedPacket* PacketToQueue);
+
+	/* Gets a packet from the the buffered packet queue for sending */
+	BufferedPacket* GetQueuedPacket();
+
+	/* Mode of the handler, Client or Server */
+	Handler::Mode Mode;
+
+	/* Time, updated by Tick */
+	float Time;
+
+private:
+	/* Set state of handler */
+	void SetState(Handler::State InState);
+
+	/* Called when handler is initialized */
+	void HandlerInitialized();
+
+	/* Used for packing outgoing packets */
+	FBitWriter OutgoingPacket;
+
+	/* Used for unpacking incoming packets */
+	FBitReader IncomingPacket;
+
+	/* Array of Handler components */
+	TArray<HandlerComponent*> HandlerComponents;
+
+	/* State of the handler */
+	Handler::State State;
+	
+	/*
+	* Packets that are buffered while handler components
+	* are being initialized
+	*/
+	TArray<BufferedPacket*> BufferedPackets;
+
+	/*
+	* Packets that are queued to be sent when handler is ticked
+	*/
+	TQueue<BufferedPacket*> QueuedPackets;
+
+	/* Reliability Handler Component */
+	ReliabilityHandlerComponent* ReliabilityComponent;
+};
+
+/*
+* This class appends or modifies incoming and outgoing packets
+* on a connection
+*/
+class PACKETHANDLER_API HandlerComponent
+{
+public:
+	/* Default initialization of data */
+	HandlerComponent();
+
+	/* Returns whether this handler is currently active */
+	virtual bool IsActive() const;
+
+	/* Return whether this handler is valid */
+	virtual bool IsValid() const PURE_VIRTUAL(Handler::Component::IsValid, return false;);
+
+	/* Returns whether this handler is initialized */
+	bool IsInitialized() const;
+
+	/* Returns whether this handler requires a tick every frame */
+	bool DoesTick() const;
+
+	/* Handles any incoming packets */
+	virtual void Incoming(FBitReader& Packet) = 0;
+
+	/* Handles any outgoing packets */
+	virtual void Outgoing(FBitWriter& Packet) = 0;
+
+	/* Initialization functionality should be placed here */
+	virtual void Initialize() = 0;
+
+	/* Tick functionality should be placed here */
+	virtual void Tick(float DeltaTime){};
+
+	/* The manager of the handler, set in initialization */
+	PacketHandler* Handler; 
+
+	/* Whether this handler is active */
+	virtual void SetActive(bool Active);
+
+protected:
+
+	/* Sets the state of the handler */
+	void SetState(Handler::Component::State State);
+
+	/* Should be called when the handler is fully initialized on both remote and local */
+	void Initialized();
+
+	/* The state of this handler */
+	Handler::Component::State State;
+
+private:
+	/* 
+	* Whether this handler is active, which dictates whether it will receive
+	* incoming and outgoing packets.
+	*/
+	bool bActive;
+
+	/* Whether this handler is fully initialized on both remote and local */
+	bool bInitialized;
+
+	/* Allows manager access to setting the Handler variable */
+	friend PacketHandler;
+};
+
+/** Packet Handler Module Interface */
+class FPacketHandlerModuleInterface : public IModuleInterface
+{
+public:
+	virtual HandlerComponent* CreateComponentInstance() PURE_VIRTUAL(FPacketHandlerModuleInterface::CreateComponentInstance,return nullptr;);
+	virtual HandlerComponent* CreateComponentInstance(FName& Options) PURE_VIRTUAL(FPacketHandlerModuleInterface::CreateComponentInstance,return nullptr;);
+};
diff --git a/Engine/Source/Runtime/PacketHandlers/ReliabilityHandlerComponent/Private/ReliabilityHandlerComponent.cpp b/Engine/Source/Runtime/PacketHandlers/ReliabilityHandlerComponent/Private/ReliabilityHandlerComponent.cpp
new file mode 100644
--- /dev/null
+++ b/Engine/Source/Runtime/PacketHandlers/ReliabilityHandlerComponent/Private/ReliabilityHandlerComponent.cpp
@@ -0,0 +1,186 @@
+#include "ReliabilityHandlerComponent.h"
+
+DEFINE_LOG_CATEGORY(PacketHandlerLog);
+IMPLEMENT_MODULE(FReliabilityHandlerComponentModuleInterface, ReliabilityHandlerComponent);
+
+// RELIABILITY
+ReliabilityHandlerComponent::ReliabilityHandlerComponent()
+: LocalPacketID(1)
+, RemotePacketID(0)
+, RemotePacketIDACKED(0)
+, ResendResolutionTime(0.5f)
+, LocalPacketIDACKED(0)
+{
+}
+
+void ReliabilityHandlerComponent::Initialize()
+{
+	SetActive(true);
+	Initialized();
+	State = Handler::Component::State::Initialized;
+}
+
+bool ReliabilityHandlerComponent::IsValid() const
+{
+	return true;
+}
+
+void ReliabilityHandlerComponent::Outgoing(FBitWriter& Packet)
+{
+	switch (State)
+	{
+		case Handler::Component::State::Initialized:
+		{
+			FBitWriter Local;
+			Local.AllowAppend(true);
+			Local.SetAllowResize(true);
+
+			if (IsActive() && IsValid())
+			{
+				bool PacketHasData = Packet.GetNumBytes() > 0;
+
+				Local.Serialize(&PacketHasData, sizeof(bool));
+				Local.SerializeIntPacked(RemotePacketID);
+
+				// Save packet for reliability
+				if (PacketHasData)
+				{
+					Local.SerializeIntPacked(LocalPacketID);
+					Local.Serialize(Packet.GetData(), Packet.GetNumBytes());
+					++LocalPacketID;
+
+					BufferedPacket* ReliablePacket = new BufferedPacket;
+
+					uint8* Data = new uint8[Packet.GetNumBytes()];
+					memcpy(Data, Packet.GetData(), Packet.GetNumBytes());
+					ReliablePacket->Data = Data;
+					ReliablePacket->Id = LocalPacketID;
+					ReliablePacket->BytesCount = Packet.GetNumBytes();
+					ReliablePacket->SendTime = Handler->Time + ResendResolutionTime;
+					BufferedPackets.Enqueue(ReliablePacket);
+				}
+
+				Packet = Local;
+				break;
+			}
+		}
+		default:
+		{
+			break;
+		}
+	}
+}
+
+void ReliabilityHandlerComponent::Incoming(FBitReader& Packet)
+{
+	switch (State)
+	{
+		case Handler::Component::State::Initialized:
+		{
+			if (IsActive() && IsValid())
+			{
+				bool PacketHasData = false;
+				Packet.Serialize(&PacketHasData, sizeof(bool));
+
+				// ACK
+				uint32 IncomingLocalPacketIDACK;
+				Packet.SerializeIntPacked(IncomingLocalPacketIDACK);
+
+				// Invalid ACK
+				if (IncomingLocalPacketIDACK > LocalPacketID)
+				{
+					Packet.Seek(Packet.GetNumBytes());
+					break;
+				}
+
+				// Set latest ACK
+				LocalPacketIDACKED = IncomingLocalPacketIDACK;
+
+				if (PacketHasData)
+				{
+					// Remote ID
+					uint32 IncomingRemotePacketID;
+					Packet.SerializeIntPacked(IncomingRemotePacketID);
+
+					// Out of sequence or duplicate packet, seek to end
+					// meaning no more bytes will be read from the packet
+					if (RemotePacketID + 1 != IncomingRemotePacketID)
+					{
+						Packet.Seek(Packet.GetNumBytes());
+						return;
+					}
+
+					// Set latest ID
+					RemotePacketID = IncomingRemotePacketID;
+				}
+
+				// Remove header from reader
+				FBitReader Copy(Packet.GetData() + (Packet.GetNumBytes() - Packet.GetBytesLeft()), Packet.GetBitsLeft());
+				Packet = Copy;
+			}
+			break;
+		}
+		default:
+		{
+			break;
+		}
+	}
+}
+
+void ReliabilityHandlerComponent::Tick(float DeltaTime)
+{
+	float CurrentTime = Handler->Time;
+
+	// Resend UNACKED packets
+	while (BufferedPackets.IsEmpty() == false)
+	{
+		BufferedPacket* Packet;
+		BufferedPackets.Peek(Packet);
+
+		// Resend
+		if (Packet->SendTime > CurrentTime)
+		{
+			BufferedPackets.Dequeue(Packet);
+
+			// Id has not been ACKED
+			if (Packet->Id < LocalPacketIDACKED)
+			{
+				Handler->QueuePacketForSending(Packet);
+			}
+			// ID has been ACKED
+			else
+			{
+				delete Packet;
+			}
+		}
+		// No packets to resend this frame
+		else
+		{
+			break;
+		}
+	}
+
+	// Need to send ACK, queue a packet
+	if (IsActive() && RemotePacketID > RemotePacketIDACKED)
+	{
+		Handler->QueuePacketForSending(new BufferedPacket);
+	}
+}
+
+void ReliabilityHandlerComponent::QueuePacketForResending(uint8* Packet, int32 Count)
+{
+	BufferedPacket* Buffered = new BufferedPacket;
+	Buffered->Data = new uint8[Count];
+	memcpy(Buffered->Data, Packet, Count);
+	Buffered->Id = LocalPacketID;
+	Buffered->SendTime = Handler->Time + ResendResolutionTime;
+	Buffered->BytesCount = Count;
+
+	BufferedPackets.Enqueue(Buffered);
+}
+
+// MODULE INTERFACE
+HandlerComponent* FReliabilityHandlerComponentModuleInterface::CreateComponentInstance()
+{
+	return new ReliabilityHandlerComponent;
+}
\ No newline at end of file
diff --git a/Engine/Source/Runtime/PacketHandlers/ReliabilityHandlerComponent/Public/ReliabilityHandlerComponent.h b/Engine/Source/Runtime/PacketHandlers/ReliabilityHandlerComponent/Public/ReliabilityHandlerComponent.h
new file mode 100644
--- /dev/null
+++ b/Engine/Source/Runtime/PacketHandlers/ReliabilityHandlerComponent/Public/ReliabilityHandlerComponent.h
@@ -0,0 +1,59 @@
+// Copyright 1998-2015 Epic Games, Inc. All Rights Reserved.
+
+#pragma once
+
+#include "PacketHandler.h"
+#include "ModuleManager.h"
+#include "Core.h"
+
+// Symmetric Stream cipher
+class RELIABILITYHANDLERCOMPONENT_API ReliabilityHandlerComponent : public HandlerComponent
+{
+public:
+	/* Initializes default data */
+	ReliabilityHandlerComponent();
+
+	/* Initializes the handler component */
+	virtual void Initialize() override;
+
+	/* Whether the handler component is valid */
+	virtual bool IsValid() const override;
+
+	/* Tick every frame */
+	virtual void Tick(float DeltaTime) override;
+
+	/* Handles any incoming packets */
+	virtual void Incoming(FBitReader& Packet) override;
+
+	/* Handles any outgoing packets */
+	virtual void Outgoing(FBitWriter& Packet) override;
+
+	/* Queues a packet for resending */
+	virtual void QueuePacketForResending(uint8* Packet, int32 Count);
+
+protected:
+	/* Buffered Packets in case they need to be resent */
+	TQueue<BufferedPacket*> BufferedPackets;
+
+	/* Latest Packet ID */
+	uint32 LocalPacketID;
+
+	/* Latest Packet ID that was ACKED */
+	uint32 LocalPacketIDACKED;
+
+	/* Latest Remote Packet ID */
+	uint32 RemotePacketID;
+	
+	/* Latest Remote Packet ID that was ACKED */
+	uint32 RemotePacketIDACKED;
+
+	/* How long to wait before resending an UNACKED packet */
+	float ResendResolutionTime;
+};
+
+/* Reliability Module Interface */
+class FReliabilityHandlerComponentModuleInterface : public FPacketHandlerModuleInterface
+{
+public:
+	virtual HandlerComponent* CreateComponentInstance() override;
+};
\ No newline at end of file
diff --git a/Engine/Source/Runtime/PacketHandlers/ReliabilityHandlerComponent/ReliabilityHandlerComponent.Build.cs b/Engine/Source/Runtime/PacketHandlers/ReliabilityHandlerComponent/ReliabilityHandlerComponent.Build.cs
new file mode 100644
--- /dev/null
+++ b/Engine/Source/Runtime/PacketHandlers/ReliabilityHandlerComponent/ReliabilityHandlerComponent.Build.cs
@@ -0,0 +1,17 @@
+// Copyright 1998-2015 Epic Games, Inc. All Rights Reserved.
+
+using UnrealBuildTool;
+using System.IO;
+
+public class ReliabilityHandlerComponent : ModuleRules
+{
+    public ReliabilityHandlerComponent(TargetInfo Target)
+    {
+        PublicDependencyModuleNames.AddRange(
+            new string[] {
+				"Core",
+                "PacketHandler",
+            }
+        );
+    }
+}
