#ue4-buildwatcher - Build fixes:

- Declare virtual destructor for PacketHandler, required on some compilers when it's deleted because it has virtual functions.
- Fix link error from multiply defined log category in same module on some builds (Win64 client mono, Win64 server mono).
- Fix warnings from initialization order differing from declaration order.

diff --git a/Engine/Source/Runtime/PacketHandlers/PacketHandler/Private/PacketHandler.cpp b/Engine/Source/Runtime/PacketHandlers/PacketHandler/Private/PacketHandler.cpp
--- a/Engine/Source/Runtime/PacketHandlers/PacketHandler/Private/PacketHandler.cpp
+++ b/Engine/Source/Runtime/PacketHandlers/PacketHandler/Private/PacketHandler.cpp
@@ -273,8 +273,8 @@ void PacketHandler::QueuePacketForSending(BufferedPacket* PacketToQueue)
 
 // HANDLER COMPONENT
 HandlerComponent::HandlerComponent()
-: bActive(false)
-, State(Handler::Component::State::UnInitialized)
+: State(Handler::Component::State::UnInitialized)
+, bActive(false)
 , bInitialized(false)
 {
 
diff --git a/Engine/Source/Runtime/PacketHandlers/PacketHandler/Public/PacketHandler.h b/Engine/Source/Runtime/PacketHandlers/PacketHandler/Public/PacketHandler.h
--- a/Engine/Source/Runtime/PacketHandlers/PacketHandler/Public/PacketHandler.h
+++ b/Engine/Source/Runtime/PacketHandlers/PacketHandler/Public/PacketHandler.h
@@ -74,6 +74,9 @@ public:
 	/* Default initialization of data */
 	PacketHandler();
 
+	/* Destructor */
+	virtual ~PacketHandler() { };
+
 	/* Handles initialization of manager */
 	virtual void Initialize(Handler::Mode Mode);
 
diff --git a/Engine/Source/Runtime/PacketHandlers/ReliabilityHandlerComponent/Private/ReliabilityHandlerComponent.cpp b/Engine/Source/Runtime/PacketHandlers/ReliabilityHandlerComponent/Private/ReliabilityHandlerComponent.cpp
--- a/Engine/Source/Runtime/PacketHandlers/ReliabilityHandlerComponent/Private/ReliabilityHandlerComponent.cpp
+++ b/Engine/Source/Runtime/PacketHandlers/ReliabilityHandlerComponent/Private/ReliabilityHandlerComponent.cpp
@@ -1,6 +1,5 @@
 #include "ReliabilityHandlerComponent.h"
 
-DEFINE_LOG_CATEGORY(PacketHandlerLog);
 IMPLEMENT_MODULE(FReliabilityHandlerComponentModuleInterface, ReliabilityHandlerComponent);
 
 // RELIABILITY
@@ -8,8 +7,8 @@ ReliabilityHandlerComponent::ReliabilityHandlerComponent()
 : LocalPacketID(1)
 , RemotePacketID(0)
 , RemotePacketIDACKED(0)
-, ResendResolutionTime(0.5f)
 , LocalPacketIDACKED(0)
+, ResendResolutionTime(0.5f)
 {
 }
 
