Encryption Handler Components

diff --git a/Engine/Source/Runtime/PacketHandlers/EncryptionComponents/AsymmetricEncryption/RSAEncryptionHandlerComponent/Private/RSAEncryptionHandlerComponent.cpp b/Engine/Source/Runtime/PacketHandlers/EncryptionComponents/AsymmetricEncryption/RSAEncryptionHandlerComponent/Private/RSAEncryptionHandlerComponent.cpp
new file mode 100644
--- /dev/null
+++ b/Engine/Source/Runtime/PacketHandlers/EncryptionComponents/AsymmetricEncryption/RSAEncryptionHandlerComponent/Private/RSAEncryptionHandlerComponent.cpp
@@ -0,0 +1,298 @@
+// Copyright 1998-2015 Epic Games, Inc. All Rights Reserved.
+
+#include "RSAEncryptionHandlerComponent.h"
+
+DEFINE_LOG_CATEGORY(PacketHandlerLog);
+IMPLEMENT_MODULE(FRSAEncryptorHandlerComponentModuleInterface, RSAEncryptionHandlerComponent);
+
+RSAEncryptionHandlerComponent::RSAEncryptionHandlerComponent(int32 InKeySizeInBits)
+{
+	KeySizeInBits = InKeySizeInBits;
+	SetActive(true);
+}
+
+void RSAEncryptionHandlerComponent::Initialize()
+{
+	Params.GenerateRandomWithKeySize(Rng, KeySizeInBits);
+	PublicKey = CryptoPP::RSA::PublicKey(Params);
+	PrivateKey = CryptoPP::RSA::PrivateKey(Params);
+
+	PrivateEncryptor = CryptoPP::RSAES_OAEP_SHA_Encryptor(PrivateKey);
+	PrivateDecryptor = CryptoPP::RSAES_OAEP_SHA_Decryptor(PrivateKey);
+
+	PrivateKeyMaxPlaintextLength = static_cast<uint32>(PrivateEncryptor.FixedMaxPlaintextLength());
+	PrivateKeyFixedCiphertextLength = static_cast<uint32>(PrivateEncryptor.FixedCiphertextLength());
+
+	SetState(ERSAEncryptionHandler::State::InitializedLocalKeys);
+}
+
+bool RSAEncryptionHandlerComponent::IsValid() const
+{
+	return PrivateKeyMaxPlaintextLength > 0;
+}
+
+void RSAEncryptionHandlerComponent::SetState(ERSAEncryptionHandler::State InState)
+{
+	State = InState;
+}
+
+void RSAEncryptionHandlerComponent::Outgoing(FBitWriter& Packet)
+{
+	switch (State)
+	{
+		case ERSAEncryptionHandler::State::InitializedLocalKeys:
+		{
+			PackLocalKey(Packet);
+			SetState(ERSAEncryptionHandler::State::InitializedLocalKeysSentLocal);
+			break;
+		}
+		case ERSAEncryptionHandler::State::InitializedLocalRemoteKeys:
+		{
+			PackLocalKey(Packet);
+			SetState(ERSAEncryptionHandler::State::InitializedLocalRemoteKeysSentLocal);
+			break;
+		}
+		case ERSAEncryptionHandler::State::Initialized:
+		{
+			if (Packet.GetNumBytes() > 0)
+			{
+				Encrypt(Packet);
+			}
+			break;
+		}
+		default:
+		{
+			break;
+		}
+	}
+}
+
+void RSAEncryptionHandlerComponent::Incoming(FBitReader& Packet)
+{
+	switch (State)
+	{
+		case ERSAEncryptionHandler::State::InitializedLocalKeys:
+		{
+			UnPackRemoteKey(Packet);
+			SetState(ERSAEncryptionHandler::State::InitializedLocalRemoteKeys);
+			break;
+		}
+		case ERSAEncryptionHandler::State::InitializedLocalKeysSentLocal:
+		{
+			UnPackRemoteKey(Packet);
+			Initialized();
+			SetState(ERSAEncryptionHandler::State::Initialized);
+			break;
+		}
+		case ERSAEncryptionHandler::State::InitializedLocalRemoteKeysSentLocal:
+		{
+			Initialized();
+			SetState(ERSAEncryptionHandler::State::Initialized);
+		}
+		case ERSAEncryptionHandler::State::Initialized:
+		{
+			Decrypt(Packet);
+			break;
+		}
+		default:
+		{
+			break;
+		}
+	}
+}
+
+void RSAEncryptionHandlerComponent::PackLocalKey(FBitWriter& Packet)
+{
+	FBitWriter Local;
+	Local.AllowAppend(true);
+	Local.SetAllowResize(true);
+
+	// Save remote public key information
+	TArray<byte> ModulusArray;
+	TArray<byte> ExponentArray;
+	SavePublicKeyModulus(ModulusArray);
+	SavePublicKeyExponent(ExponentArray);
+
+	// Pack public key info
+	Local << ModulusArray;
+	Local << ExponentArray;
+	Local.Serialize(Packet.GetData(), Packet.GetNumBytes());
+	Packet = Local;
+}
+
+void RSAEncryptionHandlerComponent::UnPackRemoteKey(FBitReader& Packet)
+{
+	// Save remote public key
+	TArray<byte> ModulusArray;
+	TArray<byte> ExponentArray;
+
+	Packet << ModulusArray;
+	Packet << ExponentArray;
+
+	CryptoPP::Integer Modulus;
+	CryptoPP::Integer Exponent;
+
+	for (int32 i = 0; i < ModulusArray.Num(); ++i)
+	{
+		Modulus.SetByte(i, ModulusArray[i]);
+	}
+
+	for (int32 i = 0; i < ExponentArray.Num(); ++i)
+	{
+		Exponent.SetByte(i, ExponentArray[i]);
+	}
+
+	RemotePublicKey.SetModulus(Modulus);
+	RemotePublicKey.SetPublicExponent(Exponent);
+	RemotePublicEncryptor = CryptoPP::RSAES_OAEP_SHA_Encryptor(RemotePublicKey);
+	RemotePublicKeyMaxPlaintextLength = static_cast<uint32>(RemotePublicEncryptor.FixedMaxPlaintextLength());
+	RemotePublicKeyFixedCiphertextLength = static_cast<uint32>(RemotePublicEncryptor.FixedCiphertextLength());
+}
+
+void RSAEncryptionHandlerComponent::Encrypt(FBitWriter& Packet)
+{
+	// Serialize size of plain text data
+	uint32 NumberOfBytesInPlaintext = static_cast<uint32>(Packet.GetNumBytes());
+
+	TArray<byte> PlainText;
+
+	// Copy byte stream to PlainText from Packet
+	for (int32 i = 0; i < Packet.GetNumBytes(); ++i)
+	{
+		PlainText.Add(Packet.GetData()[i]);
+	}
+
+	Packet.Reset();
+
+	if (NumberOfBytesInPlaintext == 0 || NumberOfBytesInPlaintext > RemotePublicKeyMaxPlaintextLength)
+	{
+		if (NumberOfBytesInPlaintext > RemotePublicKeyMaxPlaintextLength)
+		{
+			UE_LOG(PacketHandlerLog, Warning, TEXT("RSA Encryption skipped as plain text size is too large for this key size. Increase key size or send smaller packets."));
+		}
+
+		uint32 NoBytesEncrypted = 0;
+		Packet.SerializeIntPacked(NoBytesEncrypted);
+		Packet.Serialize(PlainText.GetData(), PlainText.Num());
+		return;
+	}
+
+	// Serialize invalid amount of bytes
+	Packet.SerializeIntPacked(NumberOfBytesInPlaintext);
+
+	TArray<byte> CipherText;
+
+	// The number of bytes in the plaintext is too large for the cipher text
+	EncryptWithRemotePublic(PlainText, CipherText);
+
+	// Copy byte stream to Packet from CipherText
+	for (int i = 0; i < CipherText.Num(); ++i)
+	{
+		Packet << CipherText[i];
+	}
+
+	UE_LOG(PacketHandlerLog, Log, TEXT("RSA Encrypted"));
+}
+
+void RSAEncryptionHandlerComponent::Decrypt(FBitReader& Packet)
+{
+	// Serialize size of plain text data
+	uint32 NumberOfBytesInPlaintext;
+	Packet.SerializeIntPacked(NumberOfBytesInPlaintext);
+
+	if (NumberOfBytesInPlaintext == 0 || NumberOfBytesInPlaintext > RemotePublicKeyMaxPlaintextLength)
+	{
+		// Remove header
+		FBitReader Copy(Packet.GetData() + (Packet.GetNumBytes() - Packet.GetBytesLeft()), Packet.GetBitsLeft());
+		Packet = Copy;
+
+		if (NumberOfBytesInPlaintext > RemotePublicKeyMaxPlaintextLength)
+		{
+			UE_LOG(PacketHandlerLog, Warning, TEXT("RSA Decryption skipped as cipher text size is too large for this key size. Increase key size or send smaller packets."));
+		}
+		return;
+	}
+
+	TArray<byte> PlainText;
+	TArray<byte> CipherText;
+
+	// Copy byte stream to PlainText from Packet
+	// not including the NumOfBytesPT field
+	for (int32 i = Packet.GetPosBits() / 8; i < Packet.GetNumBytes(); ++i)
+	{
+		CipherText.Add(Packet.GetData()[i]);
+	}
+
+	DecryptWithPrivate(CipherText, PlainText);
+
+	FBitReader Copy(PlainText.GetData(), NumberOfBytesInPlaintext * 8);
+	Packet = Copy;
+
+	UE_LOG(PacketHandlerLog, Log, TEXT("RSA Decrypted"));
+}
+
+void RSAEncryptionHandlerComponent::SavePublicKeyModulus(TArray<byte>& OutModulus)
+{
+	uint32 ModulusSize = static_cast<uint32>(PublicKey.GetModulus().ByteCount());
+
+	CryptoPP::Integer Modulus = PublicKey.GetModulus();
+
+	for (uint32 i = 0; i < ModulusSize; ++i)
+	{
+		OutModulus.Add(Modulus.GetByte(i));
+	}
+}
+
+void RSAEncryptionHandlerComponent::SavePublicKeyExponent(TArray<byte>& OutExponent)
+{
+	uint32 ExponentSize = static_cast<uint32>(PublicKey.GetPublicExponent().ByteCount());
+
+	CryptoPP::Integer Exponent = PublicKey.GetPublicExponent();
+
+	for (uint32 i = 0; i < ExponentSize; ++i)
+	{
+		OutExponent.Add(Exponent.GetByte(i));
+	}
+}
+
+void RSAEncryptionHandlerComponent::EncryptWithRemotePublic(const TArray<byte>& InPlainText, TArray<byte>& OutCipherText)
+{
+	OutCipherText.Reset(); // Reset array
+	OutCipherText.Reserve(RemotePublicKeyFixedCiphertextLength); // Reserve memory
+
+	byte* CipherTextBuffer = new byte[RemotePublicKeyFixedCiphertextLength];
+
+	RemotePublicEncryptor.Encrypt(Rng, InPlainText.GetData(), InPlainText.Num(), CipherTextBuffer);
+
+	for (uint32 i = 0; i < RemotePublicKeyFixedCiphertextLength; ++i)
+	{
+		OutCipherText.Add(CipherTextBuffer[i]);
+	}
+
+	delete[] CipherTextBuffer;
+}
+
+void RSAEncryptionHandlerComponent::DecryptWithPrivate(const TArray<byte>& InCipherText, TArray<byte>& OutPlainText)
+{
+	byte* PlainTextTextBuffer = new byte[PrivateKeyMaxPlaintextLength];
+
+	PrivateDecryptor.Decrypt(Rng, InCipherText.GetData(), InCipherText.Num(), PlainTextTextBuffer);
+
+	for (uint32 i = 0; i < PrivateKeyMaxPlaintextLength; ++i)
+	{
+		OutPlainText.Add(PlainTextTextBuffer[i]);
+	}
+
+	delete[] PlainTextTextBuffer;
+}
+
+// MODULE INTERFACE
+HandlerComponent* FRSAEncryptorHandlerComponentModuleInterface::CreateComponentInstance()
+{
+	return new RSAEncryptionHandlerComponent;
+}
+
+HandlerComponent* FRSAEncryptorHandlerComponentModuleInterface::CreateComponentInstance(FString& Options)
+{
+	return new RSAEncryptionHandlerComponent;
+}
\ No newline at end of file
diff --git a/Engine/Source/Runtime/PacketHandlers/EncryptionComponents/AsymmetricEncryption/RSAEncryptionHandlerComponent/Public/RSAEncryptionHandlerComponent.h b/Engine/Source/Runtime/PacketHandlers/EncryptionComponents/AsymmetricEncryption/RSAEncryptionHandlerComponent/Public/RSAEncryptionHandlerComponent.h
new file mode 100644
--- /dev/null
+++ b/Engine/Source/Runtime/PacketHandlers/EncryptionComponents/AsymmetricEncryption/RSAEncryptionHandlerComponent/Public/RSAEncryptionHandlerComponent.h
@@ -0,0 +1,127 @@
+// Copyright 1998-2015 Epic Games, Inc. All Rights Reserved.
+
+#pragma once
+
+#include "PacketHandler.h"
+
+#include "CryptoPP/5.6.2/include/rsa.h"
+#include "CryptoPP/5.6.2/include/osrng.h"
+
+/* RSA Encryptor Module Interface */
+class FRSAEncryptorHandlerComponentModuleInterface : public FPacketHandlerComponentModuleInterface
+{
+public:
+	/* Creates an instance of this component */
+	virtual HandlerComponent* CreateComponentInstance() override;
+	/* Creates an instance of this component with an FString of options */
+	virtual HandlerComponent* CreateComponentInstance(FString& Options) override;
+};
+
+
+namespace ERSAEncryptionHandler
+{
+	enum State
+	{
+		UnInitialized,
+		InitializedLocalKeys,
+		InitializedLocalKeysSentLocal,
+		InitializedLocalRemoteKeys,
+		InitializedLocalRemoteKeysSentLocal,
+		Initialized
+	};
+}
+
+/*
+* Asymmetric block cipher handler component using RSA as the cipher
+*/
+class RSAENCRYPTIONHANDLERCOMPONENT_API RSAEncryptionHandlerComponent : public HandlerComponent
+{
+public:
+	/*
+	* Initializes default data, can specify key size in bits, increasing the key size
+	* will increase the fixed cipher text and max plain text sizes
+	*/
+	RSAEncryptionHandlerComponent(int32 InKeySizeInBits = 1024);
+
+	/* Initializes the handler component */
+	virtual void Initialize() override;
+
+	/* Whether the handler component is valid */
+	virtual bool IsValid() const override;
+
+	/* Handles any incoming packets */
+	virtual void Incoming(FBitReader& Packet) override;
+
+	/* Handles any outgoing packets */
+	virtual void Outgoing(FBitWriter& Packet) override;
+
+protected:
+	/* Set the state of the handler */
+	void SetState(ERSAEncryptionHandler::State State);
+
+	/* Encrypt outgoing packets */
+	virtual void Encrypt(FBitWriter& Packet);
+
+	/* Decrypt incoming packets */
+	virtual void Decrypt(FBitReader& Packet);
+
+	/* Pack the local key into a packet */
+	virtual void PackLocalKey(FBitWriter& Packet);
+
+	/* Unpack the remote key from a packet */
+	virtual void UnPackRemoteKey(FBitReader& Packet);
+
+	/* Encrypt outgoing packet with remote connection's public key */
+	virtual void EncryptWithRemotePublic(const TArray<byte>& InPlainText, TArray<byte>& OutCipherText);
+
+	/* Decrypt incoming packet with local private key */
+	virtual void DecryptWithPrivate(const TArray<byte>& InCipherText, TArray<byte>& OutPlainText);
+
+	/* Save the public key's modulus in the provided array */
+	void SavePublicKeyModulus(TArray<byte>& OutModulus);
+
+	/* Save the public key's exponent in the provided array */
+	void SavePublicKeyExponent(TArray<byte>& OutExponent);
+
+	/** Maximum plain text length that can be encrypted with private key */
+	uint32 PrivateKeyMaxPlaintextLength;
+
+	/** Fixed cipher text length that will result of private key encryption */
+	uint32 PrivateKeyFixedCiphertextLength;
+
+	/** Maximum plain text length that can be encrypted with remote public key */
+	uint32 RemotePublicKeyMaxPlaintextLength;
+
+	/** Fixed cipher text length that will be result of RemotePublicKey encryption */
+	uint32 RemotePublicKeyFixedCiphertextLength;
+
+	/** Size of Key in Bits */
+	uint32 KeySizeInBits;
+
+	/* State of the handler */
+	ERSAEncryptionHandler::State State;
+
+	/* Random number generator used to generate key */
+	CryptoPP::AutoSeededRandomPool Rng;
+
+	/* RSA parameters for generating key */
+	CryptoPP::InvertibleRSAFunction Params;
+
+	/* Encryptor for encrypting using remote's public key */
+	CryptoPP::RSAES_OAEP_SHA_Encryptor RemotePublicEncryptor;
+
+	/* Encryptor for encrypting with local private key */
+	CryptoPP::RSAES_OAEP_SHA_Encryptor PrivateEncryptor;
+
+	/* Decryptor for decrypting with local private key */
+	CryptoPP::RSAES_OAEP_SHA_Decryptor PrivateDecryptor;
+
+	/* Local Public Key */
+	CryptoPP::RSA::PublicKey PublicKey;
+
+	/* Local Private Key */
+	CryptoPP::RSA::PrivateKey PrivateKey;
+
+	/* Remote Public Key */
+	CryptoPP::RSA::PublicKey RemotePublicKey;
+};
diff --git a/Engine/Source/Runtime/PacketHandlers/EncryptionComponents/AsymmetricEncryption/RSAEncryptionHandlerComponent/RSAEncryptionHandlerComponent.Build.cs b/Engine/Source/Runtime/PacketHandlers/EncryptionComponents/AsymmetricEncryption/RSAEncryptionHandlerComponent/RSAEncryptionHandlerComponent.Build.cs
new file mode 100644
--- /dev/null
+++ b/Engine/Source/Runtime/PacketHandlers/EncryptionComponents/AsymmetricEncryption/RSAEncryptionHandlerComponent/RSAEncryptionHandlerComponent.Build.cs
@@ -0,0 +1,27 @@
+// Copyright 1998-2015 Epic Games, Inc. All Rights Reserved.
+
+using UnrealBuildTool;
+using System.IO;
+
+public class RSAEncryptionHandlerComponent : ModuleRules
+{
+    public RSAEncryptionHandlerComponent(TargetInfo Target)
+    {
+        PublicDependencyModuleNames.AddRange(
+            new string[] {
+				"Core",
+                "PacketHandler"
+            }
+        );
+
+        if ((Target.Platform == UnrealTargetPlatform.Win64) ||
+            (Target.Platform == UnrealTargetPlatform.Win32))
+        {
+            AddThirdPartyPrivateStaticDependencies(Target,
+                "CryptoPP"
+                );
+        }
+
+        PublicIncludePathModuleNames.Add("CryptoPP");
+    }
+}
diff --git a/Engine/Source/Runtime/PacketHandlers/EncryptionComponents/EncryptionHandlerComponent/EncryptionHandlerComponent.Build.cs b/Engine/Source/Runtime/PacketHandlers/EncryptionComponents/EncryptionHandlerComponent/EncryptionHandlerComponent.Build.cs
new file mode 100644
--- /dev/null
+++ b/Engine/Source/Runtime/PacketHandlers/EncryptionComponents/EncryptionHandlerComponent/EncryptionHandlerComponent.Build.cs
@@ -0,0 +1,20 @@
+// Copyright 1998-2015 Epic Games, Inc. All Rights Reserved.
+
+using UnrealBuildTool;
+using System.IO;
+
+public class EncryptionHandlerComponent : ModuleRules
+{
+    public EncryptionHandlerComponent(TargetInfo Target)
+    {
+        PublicDependencyModuleNames.AddRange(
+            new string[] {
+				"Core",
+                "PacketHandler",
+                "BlockEncryptionHandlerComponent",
+                "StreamEncryptionHandlerComponent",
+                "RSAEncryptionHandlerComponent",
+            }
+        );
+    }
+}
diff --git a/Engine/Source/Runtime/PacketHandlers/EncryptionComponents/EncryptionHandlerComponent/Private/EncryptionHandlerComponent.cpp b/Engine/Source/Runtime/PacketHandlers/EncryptionComponents/EncryptionHandlerComponent/Private/EncryptionHandlerComponent.cpp
new file mode 100644
--- /dev/null
+++ b/Engine/Source/Runtime/PacketHandlers/EncryptionComponents/EncryptionHandlerComponent/Private/EncryptionHandlerComponent.cpp
@@ -0,0 +1,151 @@
+// Copyright 1998-2015 Epic Games, Inc. All Rights Reserved.
+
+#include "EncryptionHandlerComponent.h"
+#include "BlockEncryptionHandlerComponent.h"
+#include "RSAEncryptionHandlerComponent.h"
+
+IMPLEMENT_MODULE(FEncryptionHandlerComponentModuleInterface, EncryptionHandlerComponent);
+
+// ENCRYPTION PROCESSOR
+EncryptionHandlerComponent::EncryptionHandlerComponent(HandlerComponent* InSymmetricHandlerComponent, HandlerComponent* InAsymmetricHandlerComponent)
+{
+	if (InSymmetricHandlerComponent != nullptr)
+	{
+		SymmetricHandlerComponent = InSymmetricHandlerComponent;
+	} 
+	else
+	{
+		SymmetricHandlerComponent = new BlockEncryptionHandlerComponent;
+	}
+
+	if (InAsymmetricHandlerComponent != nullptr)
+	{
+		AsymmetricHandlerComponent = InAsymmetricHandlerComponent;
+	} 
+	else
+	{
+		AsymmetricHandlerComponent = new RSAEncryptionHandlerComponent;
+	}
+
+	SetActive(true);
+}
+
+void EncryptionHandlerComponent::Initialize()
+{
+	SymmetricHandlerComponent->Handler = Handler;
+	SymmetricHandlerComponent->Initialize();
+
+	AsymmetricHandlerComponent->Handler = Handler;
+	AsymmetricHandlerComponent->Initialize();
+
+	SetState(EEncryptionHandler::State::InitializingAsymmetric);
+}
+
+bool EncryptionHandlerComponent::IsValid() const
+{
+	return SymmetricHandlerComponent->IsValid() && AsymmetricHandlerComponent->IsValid();
+}
+
+void EncryptionHandlerComponent::SetState(EEncryptionHandler::State InState)
+{
+	State = InState;
+}
+
+void EncryptionHandlerComponent::Incoming(FBitReader& Packet)
+{
+	switch (State)
+	{
+		case EEncryptionHandler::State::InitializingAsymmetric:
+		{
+			AsymmetricHandlerComponent->Incoming(Packet);
+
+			if (AsymmetricHandlerComponent->IsInitialized())
+			{
+				if (Handler->Mode == Handler::Mode::Client)
+				{
+					SetState(EEncryptionHandler::State::InitializingSymmetric);
+				} 
+				else if (Handler->Mode == Handler::Mode::Server)
+				{
+					SymmetricHandlerComponent->Incoming(Packet);
+				}
+			}
+			break;
+		}
+		case EEncryptionHandler::State::InitializingSymmetric:
+		{
+			AsymmetricHandlerComponent->Incoming(Packet);
+			SymmetricHandlerComponent->Incoming(Packet);
+			break;
+		}
+		case EEncryptionHandler::State::Initialized:
+		{
+			SymmetricHandlerComponent->Incoming(Packet);
+			break;
+		}
+		default:
+		{
+			break;
+		}
+	}
+
+	if (State != EEncryptionHandler::State::Initialized && SymmetricHandlerComponent->IsInitialized())
+	{
+		AsymmetricHandlerComponent->SetActive(false);
+		Initialized();
+		SetState(EEncryptionHandler::State::Initialized);
+	}
+}
+
+void EncryptionHandlerComponent::Outgoing(FBitWriter& Packet)
+{
+	switch (State)
+	{
+		case EEncryptionHandler::State::InitializingAsymmetric:
+		{
+			AsymmetricHandlerComponent->Outgoing(Packet);
+
+			if (AsymmetricHandlerComponent->IsInitialized())
+			{
+				SetState(EEncryptionHandler::State::InitializingSymmetric);
+			}
+			break;
+		}
+		case EEncryptionHandler::State::InitializingSymmetric:
+		{
+			SymmetricHandlerComponent->Outgoing(Packet);
+			AsymmetricHandlerComponent->Outgoing(Packet);
+			break;
+		}
+		case EEncryptionHandler::State::Initialized:
+		{
+			if(Packet.GetNumBytes() > 0)
+			{
+				SymmetricHandlerComponent->Outgoing(Packet);
+			}
+			break;
+		}
+		default:
+		{
+			break;
+		}
+	}
+
+	if (State != EEncryptionHandler::State::Initialized && SymmetricHandlerComponent->IsInitialized())
+	{
+		AsymmetricHandlerComponent->SetActive(false);
+		Initialized();
+		SetState(EEncryptionHandler::State::Initialized);
+	}
+}
+
+//MODULE INTERFACE
+HandlerComponent* FEncryptionHandlerComponentModuleInterface::CreateComponentInstance()
+{
+	return new EncryptionHandlerComponent;
+}
+
+HandlerComponent* FEncryptionHandlerComponentModuleInterface::CreateComponentInstance(FString& Options)
+{
+	return new EncryptionHandlerComponent;
+}
\ No newline at end of file
diff --git a/Engine/Source/Runtime/PacketHandlers/EncryptionComponents/EncryptionHandlerComponent/Public/EncryptionHandlerComponent.h b/Engine/Source/Runtime/PacketHandlers/EncryptionComponents/EncryptionHandlerComponent/Public/EncryptionHandlerComponent.h
new file mode 100644
--- /dev/null
+++ b/Engine/Source/Runtime/PacketHandlers/EncryptionComponents/EncryptionHandlerComponent/Public/EncryptionHandlerComponent.h
@@ -0,0 +1,65 @@
+// Copyright 1998-2015 Epic Games, Inc. All Rights Reserved.
+
+#pragma once
+
+#include "PacketHandler.h"
+#include "ModuleManager.h"
+#include "Core.h"
+
+/** Encryption Handler Component Module Interface */
+class FEncryptionHandlerComponentModuleInterface : public FPacketHandlerComponentModuleInterface
+{
+	public:
+	/* Creates an instance of this component */
+	virtual HandlerComponent* CreateComponentInstance() override;
+	/* Creates an instance of this component with an FString of options */
+	virtual HandlerComponent* CreateComponentInstance(FString& Options) override;
+};
+
+namespace EEncryptionHandler
+{
+	enum State
+	{
+		UnInitialized,
+		InitializingAsymmetric,
+		InitializingSymmetric,
+		Initialized
+	};
+}
+
+/*
+* This class uses Asymmetric encryption to send an encrypted symmetric encryption key
+* to the remote side. Defaults are RSA for asymmetric encryption and XOR stream cipher
+* for symmetric encryption.
+*/
+class ENCRYPTIONHANDLERCOMPONENT_API EncryptionHandlerComponent : public HandlerComponent
+{
+public:
+	/* Initializes the default encryption handler components unless specified by the user */ 
+	EncryptionHandlerComponent(HandlerComponent* SymmetricHandlerComponent = nullptr, HandlerComponent* AsymmetricHandlerComponent = nullptr);
+
+	/* Initializes the handler component */
+	virtual void Initialize() override;
+
+	/* Whether the handler component is valid */
+	virtual bool IsValid() const override;
+
+	/* Handles any incoming packets */
+	virtual void Incoming(FBitReader& Packet) override;
+
+	/* Handles any outgoing packets */
+	virtual void Outgoing(FBitWriter& Packet) override;
+
+protected:
+	/* Sets the state of the handler */
+	void SetState(EEncryptionHandler::State State);
+
+	/* The state of this handler */
+	EEncryptionHandler::State State;
+
+	/* Symmetric encryption handler component */
+	HandlerComponent* SymmetricHandlerComponent;
+
+	/* Asymmetric encryption handler component */
+	HandlerComponent* AsymmetricHandlerComponent;
+};
diff --git a/Engine/Source/Runtime/PacketHandlers/EncryptionComponents/SymmetricEncryption/BlockEncryption/AESBlockEncryptor/AESBlockEncryptor.Build.cs b/Engine/Source/Runtime/PacketHandlers/EncryptionComponents/SymmetricEncryption/BlockEncryption/AESBlockEncryptor/AESBlockEncryptor.Build.cs
new file mode 100644
--- /dev/null
+++ b/Engine/Source/Runtime/PacketHandlers/EncryptionComponents/SymmetricEncryption/BlockEncryption/AESBlockEncryptor/AESBlockEncryptor.Build.cs
@@ -0,0 +1,27 @@
+// Copyright 1998-2015 Epic Games, Inc. All Rights Reserved.
+
+using UnrealBuildTool;
+using System.IO;
+
+public class AESBlockEncryptor : ModuleRules
+{
+    public AESBlockEncryptor(TargetInfo Target)
+    {
+        PublicDependencyModuleNames.AddRange(
+            new string[] {
+				"Core",
+                "BlockEncryptionHandlerComponent"
+            }
+        );
+
+        if ((Target.Platform == UnrealTargetPlatform.Win64) ||
+            (Target.Platform == UnrealTargetPlatform.Win32))
+        {
+            AddThirdPartyPrivateStaticDependencies(Target,
+                "CryptoPP"
+                );
+        }
+
+        PublicIncludePathModuleNames.Add("CryptoPP");
+    }
+}
diff --git a/Engine/Source/Runtime/PacketHandlers/EncryptionComponents/SymmetricEncryption/BlockEncryption/AESBlockEncryptor/Private/AESBlockEncryptor.cpp b/Engine/Source/Runtime/PacketHandlers/EncryptionComponents/SymmetricEncryption/BlockEncryption/AESBlockEncryptor/Private/AESBlockEncryptor.cpp
new file mode 100644
--- /dev/null
+++ b/Engine/Source/Runtime/PacketHandlers/EncryptionComponents/SymmetricEncryption/BlockEncryption/AESBlockEncryptor/Private/AESBlockEncryptor.cpp
@@ -0,0 +1,47 @@
+// Copyright 1998-2015 Epic Games, Inc. All Rights Reserved.
+
+#include "AESBlockEncryptor.h"
+
+IMPLEMENT_MODULE(FAESBlockEncryptorModuleInterface, AESBlockEncryptor);
+
+// MODULE INTERFACE
+BlockEncryptor* FAESBlockEncryptorModuleInterface::CreateBlockEncryptorInstance()
+{
+	return new AESBlockEncryptor;
+}
+
+// AES
+void AESBlockEncryptor::Initialize(TArray<byte>* InKey)
+{
+	Key = InKey;
+
+	if (Key->Num() != 16 && Key->Num() != 32)
+	{
+		LowLevelFatalError(TEXT("Incorrect key size. %i size chosen"), Key->Num());
+	}
+
+	Encryptor = CryptoPP::AES::Encryption(Key->GetData(), Key->Num());
+	Decryptor = CryptoPP::AES::Decryption(Key->GetData(), Key->Num());
+
+	FixedBlockSize = 16;
+}
+
+void AESBlockEncryptor::EncryptBlock(byte* Block)
+{
+	byte* Output = new byte[FixedBlockSize];
+	Encryptor.ProcessBlock(Block, Output);
+	memcpy(Block, Output, FixedBlockSize);
+	delete[] Output;
+
+	//UE_LOG(PacketHandlerLog, Log, TEXT("AES Encrypted"));
+}
+
+void AESBlockEncryptor::DecryptBlock(byte* Block)
+{
+	byte* Output = new byte[FixedBlockSize];
+	Decryptor.ProcessBlock(Block, Output);
+	memcpy(Block, Output, FixedBlockSize);
+	delete[] Output;
+
+	//UE_LOG(PacketHandlerLog, Log, TEXT("AES Decrypted"));
+}
diff --git a/Engine/Source/Runtime/PacketHandlers/EncryptionComponents/SymmetricEncryption/BlockEncryption/AESBlockEncryptor/Public/AESBlockEncryptor.h b/Engine/Source/Runtime/PacketHandlers/EncryptionComponents/SymmetricEncryption/BlockEncryption/AESBlockEncryptor/Public/AESBlockEncryptor.h
new file mode 100644
--- /dev/null
+++ b/Engine/Source/Runtime/PacketHandlers/EncryptionComponents/SymmetricEncryption/BlockEncryption/AESBlockEncryptor/Public/AESBlockEncryptor.h
@@ -0,0 +1,36 @@
+// Copyright 1998-2015 Epic Games, Inc. All Rights Reserved.
+
+#pragma once
+
+#include "BlockEncryptionHandlerComponent.h"
+#include "CryptoPP/5.6.2/include/aes.h"
+
+/* AES Block Encryptor Module Interface */
+class FAESBlockEncryptorModuleInterface : public FBlockEncryptorModuleInterface
+{
+	virtual BlockEncryptor* CreateBlockEncryptorInstance() override;
+};
+
+/*
+* AES Block encryption
+*/
+class AESBLOCKENCRYPTOR_API AESBlockEncryptor : public BlockEncryptor
+{
+public:
+	/* Initialized the encryptor */
+	void Initialize(TArray<byte>* Key) override;
+
+	/* Encrypts outgoing packets */
+	void EncryptBlock(byte* Block) override;
+
+	/* Decrypts incoming packets */
+	void DecryptBlock(byte* Block) override;
+
+	/* Get the default key size for this encryptor */
+	uint32 GetDefaultKeySize() { return 16; }
+
+private:
+	/* Encryptors for AES */
+	CryptoPP::AES::Encryption Encryptor;
+	CryptoPP::AES::Decryption Decryptor;
+};
\ No newline at end of file
diff --git a/Engine/Source/Runtime/PacketHandlers/EncryptionComponents/SymmetricEncryption/BlockEncryption/BlockEncryptionHandlerComponent/BlockEncryptionHandlerComponent.Build.cs b/Engine/Source/Runtime/PacketHandlers/EncryptionComponents/SymmetricEncryption/BlockEncryption/BlockEncryptionHandlerComponent/BlockEncryptionHandlerComponent.Build.cs
new file mode 100644
--- /dev/null
+++ b/Engine/Source/Runtime/PacketHandlers/EncryptionComponents/SymmetricEncryption/BlockEncryption/BlockEncryptionHandlerComponent/BlockEncryptionHandlerComponent.Build.cs
@@ -0,0 +1,31 @@
+// Copyright 1998-2015 Epic Games, Inc. All Rights Reserved.
+
+using UnrealBuildTool;
+using System.IO;
+
+public class BlockEncryptionHandlerComponent : ModuleRules
+{
+    public BlockEncryptionHandlerComponent(TargetInfo Target)
+    {
+        PublicDependencyModuleNames.AddRange(
+            new string[] {
+				"Core",
+                "PacketHandler",
+                "XORBlockEncryptor",
+            }
+        );
+
+        CircularlyReferencedDependentModules.Add("XORBlockEncryptor");
+
+        if ((Target.Platform == UnrealTargetPlatform.Win64) ||
+        (Target.Platform == UnrealTargetPlatform.Win32))
+        {
+            AddThirdPartyPrivateStaticDependencies(Target,
+                "CryptoPP"
+                );
+        }
+
+        PublicIncludePathModuleNames.Add("CryptoPP");
+
+    }
+}
diff --git a/Engine/Source/Runtime/PacketHandlers/EncryptionComponents/SymmetricEncryption/BlockEncryption/BlockEncryptionHandlerComponent/Private/BlockEncryptionHandlerComponent.cpp b/Engine/Source/Runtime/PacketHandlers/EncryptionComponents/SymmetricEncryption/BlockEncryption/BlockEncryptionHandlerComponent/Private/BlockEncryptionHandlerComponent.cpp
new file mode 100644
--- /dev/null
+++ b/Engine/Source/Runtime/PacketHandlers/EncryptionComponents/SymmetricEncryption/BlockEncryption/BlockEncryptionHandlerComponent/Private/BlockEncryptionHandlerComponent.cpp
@@ -0,0 +1,235 @@
+// Copyright 1998-2015 Epic Games, Inc. All Rights Reserved.
+
+#include "BlockEncryptionHandlerComponent.h"
+#include "XORBlockEncryptor.h"
+
+IMPLEMENT_MODULE(FBlockEncryptionHandlerComponentModuleInterface, BlockEncryptionHandlerComponent);
+
+///////////////////////////////////////////////////
+// BLOCK ENCRYPTION
+///////////////////////////////////////////////////
+
+// Block Encryption Handler
+BlockEncryptionHandlerComponent::BlockEncryptionHandlerComponent(BlockEncryptor* InEncryptor, uint32 InKeySizeInBytes)
+{
+	if (InEncryptor != nullptr)
+	{
+		Encryptor = InEncryptor;
+	} else
+	{
+		Encryptor = new XORBlockEncryptor;
+	}
+
+	// Set key size, if no key size passed in then get default key size for the encryptor
+	KeySizeInBytes = InKeySizeInBytes != 0 ? InKeySizeInBytes : Encryptor->GetDefaultKeySize();
+}
+
+void BlockEncryptionHandlerComponent::Initialize()
+{
+	if (Handler->Mode == Handler::Mode::Client)
+	{
+		// Generate the key
+		CryptoPP::AutoSeededRandomPool Rng;
+
+		Key.Reserve(KeySizeInBytes);
+		byte* Secret = new byte[KeySizeInBytes];
+
+		Rng.GenerateBlock(Secret, KeySizeInBytes);
+
+		for (uint32 i = 0; i < KeySizeInBytes; ++i)
+		{
+			Key.Add(Secret[i]);
+		}
+
+		Encryptor->Initialize(&Key);
+		SetState(Handler::Component::State::InitializedOnLocal);
+	}
+	SetActive(true);
+}
+
+bool BlockEncryptionHandlerComponent::IsValid() const
+{
+	return Key.Num() > 0;
+}
+
+void BlockEncryptionHandlerComponent::Incoming(FBitReader& Packet)
+{
+	switch (State)
+	{
+		case Handler::Component::State::InitializedOnLocal:
+		case Handler::Component::State::UnInitialized:
+		{
+			TArray<byte> ReceivedKey;
+			Packet << ReceivedKey;
+
+			// Set key
+			Key = ReceivedKey;
+			KeySizeInBytes = uint32(Key.Num());
+
+			Encryptor->Initialize(&Key);
+			SetState(Handler::Component::State::Initialized);
+			Initialized();
+			break;
+		}
+		case Handler::Component::State::Initialized:
+		{
+			// Handle this packet
+			if (IsValid() && Packet.GetNumBytes() > 0)
+			{
+				DecryptBlock(Packet);
+			}
+			break;
+		}
+		default:
+		{
+			break;
+		}
+	}
+}
+
+void BlockEncryptionHandlerComponent::Outgoing(FBitWriter& Packet)
+{
+	switch (State)
+	{
+		case Handler::Component::State::UnInitialized:
+		{
+													 break;
+		}
+		// Send packet initializing remote handler
+		case Handler::Component::State::InitializedOnLocal:
+		{
+			FBitWriter Local;
+			Local.AllowAppend(true);
+			Local.SetAllowResize(true);
+
+			Local << Key;
+			Local.Serialize(Packet.GetData(), Packet.GetNumBytes());
+			Packet = Local;
+
+			SetState(Handler::Component::State::Initialized);
+			SetActive(true);
+			Initialized();
+			break;
+		}
+		case Handler::Component::State::Initialized:
+		{
+			// Handle this packet
+			if (IsValid() && Packet.GetNumBytes() > 0)
+			{
+				EncryptBlock(Packet);
+			}
+			break;
+		}
+		default:
+		{
+			break;
+		}
+	}
+}
+
+void BlockEncryptionHandlerComponent::EncryptBlock(FBitWriter& Packet)
+{
+	uint32 BlockSize = Encryptor->GetFixedBlockSize();
+
+	TArray<byte> Block;
+	Block.Reserve((Packet.GetNumBytes() / BlockSize) + 1);
+
+	// Copy data
+	for (int32 i = 0; i < Packet.GetNumBytes(); ++i)
+	{
+		Block.Add(Packet.GetData()[i]);
+	}
+
+	// Add padding
+	while (Block.Num() % BlockSize != 0)
+	{
+		Block.Add(0x00);
+	}
+
+	int32 BlockCount = Block.Num();
+
+	for (int32 i = 0; i < int32(BlockCount / BlockSize); ++i)
+	{
+		Encryptor->EncryptBlock(Block.GetData() + i * KeySizeInBytes);
+	}
+
+	uint32 PacketSizeBeforeEncryption = static_cast<uint32>(Packet.GetNumBytes());
+
+	Packet.Reset();
+
+	// Serialize size of packet before encryption
+	Packet.SerializeIntPacked(PacketSizeBeforeEncryption);
+
+	// Copy data back
+	for (int32 i = 0; i < BlockCount; ++i)
+	{
+		Packet << Block[i];
+	}
+}
+
+void BlockEncryptionHandlerComponent::DecryptBlock(FBitReader& Packet)
+{
+	uint32 BlockSize = Encryptor->GetFixedBlockSize();
+
+	TArray<byte> Block;
+	Block.Reserve((Packet.GetNumBytes() / BlockSize) + 1);
+
+	uint32 PacketSizeBeforeEncryption;
+	Packet.SerializeIntPacked(PacketSizeBeforeEncryption);
+
+	// Copy data
+	for (int32 i = Packet.GetPosBits() / 8; i < Packet.GetNumBytes(); ++i)
+	{
+		Block.Add(Packet.GetData()[i]);
+	}
+
+	int32 BlockCount = Block.Num();
+
+	// Not a valid sized block to decrypt, seek to end to discontinue handling
+	if (BlockCount % BlockSize != 0)
+	{
+		Packet.Seek(Packet.GetNumBytes());
+		return;
+	}
+
+	for (int32 i = 0; i < int32(BlockCount / BlockSize); ++i)
+	{
+		Encryptor->DecryptBlock(Block.GetData() + i * KeySizeInBytes);
+	}
+
+	// Packet size before encryption should always be
+	// smaller or equal to the block size
+	if (PacketSizeBeforeEncryption > static_cast<uint32>(Block.Num()))
+	{
+		abort();
+	}
+
+	FBitReader Copy(Block.GetData(), PacketSizeBeforeEncryption * 8);
+	Packet = Copy;
+
+	FString KeyString;
+
+	for (int32 i = 0; i < Key.Num(); ++i)
+	{
+		KeyString.AppendChar(Key[i]);
+	}
+}
+
+// BLOCK ENCRYPTOR
+uint32 BlockEncryptor::GetFixedBlockSize()
+{
+	return FixedBlockSize;
+}
+
+// MODULE INTERFACE
+HandlerComponent* FBlockEncryptionHandlerComponentModuleInterface::CreateComponentInstance()
+{
+	return new BlockEncryptionHandlerComponent;
+}
+
+HandlerComponent* FBlockEncryptionHandlerComponentModuleInterface::CreateComponentInstance(FString& Options)
+{
+	TSharedPtr<IModuleInterface> Interface = FModuleManager::Get().LoadModule(FName(*Options));
+	TSharedPtr<FBlockEncryptorModuleInterface> BlockEncryptorInterface(static_cast<FBlockEncryptorModuleInterface*>(&(*Interface)));
+	return new BlockEncryptionHandlerComponent(BlockEncryptorInterface->CreateBlockEncryptorInstance());
+}
\ No newline at end of file
diff --git a/Engine/Source/Runtime/PacketHandlers/EncryptionComponents/SymmetricEncryption/BlockEncryption/BlockEncryptionHandlerComponent/Public/BlockEncryptionHandlerComponent.h b/Engine/Source/Runtime/PacketHandlers/EncryptionComponents/SymmetricEncryption/BlockEncryption/BlockEncryptionHandlerComponent/Public/BlockEncryptionHandlerComponent.h
new file mode 100644
--- /dev/null
+++ b/Engine/Source/Runtime/PacketHandlers/EncryptionComponents/SymmetricEncryption/BlockEncryption/BlockEncryptionHandlerComponent/Public/BlockEncryptionHandlerComponent.h
@@ -0,0 +1,88 @@
+// Copyright 1998-2015 Epic Games, Inc. All Rights Reserved.
+
+#pragma once
+
+#include "PacketHandler.h"
+
+#include "CryptoPP/5.6.2/include/osrng.h"
+
+/*
+* Abstract class for block encryptors
+*/
+class BLOCKENCRYPTIONHANDLERCOMPONENT_API BlockEncryptor
+{
+public:
+	/* Initialized the encryptor */
+	virtual void Initialize(TArray<byte>* Key) PURE_VIRTUAL(BlockEncryptor::Initialize, );
+
+	/* Encrypts block */
+	virtual void EncryptBlock(byte* Block) PURE_VIRTUAL(BlockEncryptor::EncryptBlock, );
+
+	/* Decrypts incoming packets */
+	virtual void DecryptBlock(byte* Block) PURE_VIRTUAL(BlockEncryptor::DecryptBlock, );
+
+	/* Gets the default key size for this encryptor */
+	virtual uint32 GetDefaultKeySize() PURE_VIRTUAL(BlockEncryptor::GetDefaultKeySize, return 0;);
+
+	/* Gets the fixed size of the block */
+	uint32 GetFixedBlockSize();
+
+protected:
+	/* Fixed size of block */
+	uint32 FixedBlockSize;
+
+	/* Key, this is initialized when Initialize() is called */
+	TArray<byte>* Key;
+};
+
+/* Block Encryptor Module Interface */
+class FBlockEncryptorModuleInterface : public IModuleInterface
+{
+public:
+	virtual BlockEncryptor* CreateBlockEncryptorInstance() = 0;
+};
+
+/*
+* Symmetric block cipher handler component using AES as the cipher
+*/
+class BLOCKENCRYPTIONHANDLERCOMPONENT_API BlockEncryptionHandlerComponent : public HandlerComponent
+{
+public:
+	/* Initializes default data, can provide the key size in bytes */
+	BlockEncryptionHandlerComponent(BlockEncryptor* Encryptor = nullptr, uint32 KeySizeInBytes = 0);
+
+	/* Initializes the handler component */
+	virtual void Initialize() override;
+
+	/* Whether the handler component is valid */
+	virtual bool IsValid() const override;
+
+	/* Handles any incoming packets */
+	virtual void Incoming(FBitReader& Packet) override;
+
+	/* Handles any outgoing packets */
+	virtual void Outgoing(FBitWriter& Packet) override;
+
+protected:
+	void EncryptBlock(FBitWriter& Packet);
+	void DecryptBlock(FBitReader& Packet);
+
+	/* The encryptor used */
+	BlockEncryptor* Encryptor;
+
+	/** Size of Key in Bits */
+	uint32 KeySizeInBytes;
+
+	// Key used for symmetric encryption
+	TArray<byte> Key;
+};
+
+/* Block Encryption Handler Component Module Interface */
+class FBlockEncryptionHandlerComponentModuleInterface : public FPacketHandlerComponentModuleInterface
+{
+public:
+	/* Creates an instance of this component */
+	virtual HandlerComponent* CreateComponentInstance() override;
+	/* Creates an instance of this component with an FString of options */
+	virtual HandlerComponent* CreateComponentInstance(FString& Options) override;
+};
diff --git a/Engine/Source/Runtime/PacketHandlers/EncryptionComponents/SymmetricEncryption/BlockEncryption/BlowFishBlockEncryptor/BlowFishBlockEncryptor.Build.cs b/Engine/Source/Runtime/PacketHandlers/EncryptionComponents/SymmetricEncryption/BlockEncryption/BlowFishBlockEncryptor/BlowFishBlockEncryptor.Build.cs
new file mode 100644
--- /dev/null
+++ b/Engine/Source/Runtime/PacketHandlers/EncryptionComponents/SymmetricEncryption/BlockEncryption/BlowFishBlockEncryptor/BlowFishBlockEncryptor.Build.cs
@@ -0,0 +1,27 @@
+// Copyright 1998-2015 Epic Games, Inc. All Rights Reserved.
+
+using UnrealBuildTool;
+using System.IO;
+
+public class BlowFishBlockEncryptor : ModuleRules
+{
+    public BlowFishBlockEncryptor(TargetInfo Target)
+    {
+        PublicDependencyModuleNames.AddRange(
+            new string[] {
+				"Core",
+                "BlockEncryptionHandlerComponent"
+            }
+        );
+
+        if ((Target.Platform == UnrealTargetPlatform.Win64) ||
+            (Target.Platform == UnrealTargetPlatform.Win32))
+        {
+            AddThirdPartyPrivateStaticDependencies(Target,
+                "CryptoPP"
+                );
+        }
+
+        PublicIncludePathModuleNames.Add("CryptoPP");
+    }
+}
diff --git a/Engine/Source/Runtime/PacketHandlers/EncryptionComponents/SymmetricEncryption/BlockEncryption/BlowFishBlockEncryptor/Private/BlowFishBlockEncryptor.cpp b/Engine/Source/Runtime/PacketHandlers/EncryptionComponents/SymmetricEncryption/BlockEncryption/BlowFishBlockEncryptor/Private/BlowFishBlockEncryptor.cpp
new file mode 100644
--- /dev/null
+++ b/Engine/Source/Runtime/PacketHandlers/EncryptionComponents/SymmetricEncryption/BlockEncryption/BlowFishBlockEncryptor/Private/BlowFishBlockEncryptor.cpp
@@ -0,0 +1,47 @@
+// Copyright 1998-2015 Epic Games, Inc. All Rights Reserved.
+
+#include "BlowFishBlockEncryptor.h"
+
+IMPLEMENT_MODULE(FBlowFishBlockEncryptorModuleInterface, BlowFishBlockEncryptor);
+
+// MODULE INTERFACE
+BlockEncryptor* FBlowFishBlockEncryptorModuleInterface::CreateBlockEncryptorInstance()
+{
+	return new BlowFishBlockEncryptor;
+}
+
+// BLOWFISH
+void BlowFishBlockEncryptor::Initialize(TArray<byte>* InKey)
+{
+	Key = InKey;
+
+	if (Key->Num() <= 4 || Key->Num() >= 54 || Key->Num() % 8 != 0)
+	{
+		LowLevelFatalError(TEXT("Incorrect key size. %i size chosen"), Key->Num());
+	}
+
+	Encryptor = CryptoPP::Blowfish::Encryption(Key->GetData(), Key->Num());
+	Decryptor = CryptoPP::Blowfish::Decryption(Key->GetData(), Key->Num());
+
+	FixedBlockSize = 8;
+}
+
+void BlowFishBlockEncryptor::EncryptBlock(byte* Block)
+{
+	byte* Output = new byte[FixedBlockSize];
+	Encryptor.ProcessBlock(Block, Output);
+	memcpy(Block, Output, FixedBlockSize);
+	delete[] Output;
+
+	//UE_LOG(PacketHandlerLog, Log, TEXT("BlowFish Encrypted"));
+}
+
+void BlowFishBlockEncryptor::DecryptBlock(byte* Block)
+{
+	byte* Output = new byte[FixedBlockSize];
+	Decryptor.ProcessBlock(Block, Output);
+	memcpy(Block, Output, FixedBlockSize);
+	delete[] Output;
+
+	//UE_LOG(PacketHandlerLog, Log, TEXT("BlowFish Decrypted"));
+}
\ No newline at end of file
diff --git a/Engine/Source/Runtime/PacketHandlers/EncryptionComponents/SymmetricEncryption/BlockEncryption/BlowFishBlockEncryptor/Public/BlowFishBlockEncryptor.h b/Engine/Source/Runtime/PacketHandlers/EncryptionComponents/SymmetricEncryption/BlockEncryption/BlowFishBlockEncryptor/Public/BlowFishBlockEncryptor.h
new file mode 100644
--- /dev/null
+++ b/Engine/Source/Runtime/PacketHandlers/EncryptionComponents/SymmetricEncryption/BlockEncryption/BlowFishBlockEncryptor/Public/BlowFishBlockEncryptor.h
@@ -0,0 +1,37 @@
+// Copyright 1998-2015 Epic Games, Inc. All Rights Reserved.
+
+#pragma once
+
+#include "BlockEncryptionHandlerComponent.h"
+#include "CryptoPP/5.6.2/include/blowfish.h"
+
+/* BlowFish Block Encryptor Module Interface */
+class FBlowFishBlockEncryptorModuleInterface : public FBlockEncryptorModuleInterface
+{
+	virtual BlockEncryptor* CreateBlockEncryptorInstance() override;
+};
+
+
+/*
+* BlowFish Block encryption
+*/
+class BLOWFISHBLOCKENCRYPTOR_API BlowFishBlockEncryptor : public BlockEncryptor
+{
+public:
+	/* Initialized the encryptor */
+	void Initialize(TArray<byte>* Key) override;
+
+	/* Encrypts outgoing packets */
+	void EncryptBlock(byte* Block) override;
+
+	/* Decrypts incoming packets */
+	void DecryptBlock(byte* Block) override;
+
+	/* Get the default key size for this encryptor */
+	uint32 GetDefaultKeySize() { return 8; }
+
+private:
+	/* Encryptors for AES */
+	CryptoPP::Blowfish::Encryption Encryptor;
+	CryptoPP::Blowfish::Decryption Decryptor;
+};
\ No newline at end of file
diff --git a/Engine/Source/Runtime/PacketHandlers/EncryptionComponents/SymmetricEncryption/BlockEncryption/TwoFishBlockEncryptor/Private/TwoFishBlockEncryptor.cpp b/Engine/Source/Runtime/PacketHandlers/EncryptionComponents/SymmetricEncryption/BlockEncryption/TwoFishBlockEncryptor/Private/TwoFishBlockEncryptor.cpp
new file mode 100644
--- /dev/null
+++ b/Engine/Source/Runtime/PacketHandlers/EncryptionComponents/SymmetricEncryption/BlockEncryption/TwoFishBlockEncryptor/Private/TwoFishBlockEncryptor.cpp
@@ -0,0 +1,47 @@
+// Copyright 1998-2015 Epic Games, Inc. All Rights Reserved.
+
+#include "TwoFishBlockEncryptor.h"
+
+IMPLEMENT_MODULE(FTwoFishBlockEncryptorModuleInterface, TwoFishBlockEncryptor);
+
+// MODULE INTERFACE
+BlockEncryptor* FTwoFishBlockEncryptorModuleInterface::CreateBlockEncryptorInstance()
+{
+	return new TwoFishBlockEncryptor;
+}
+
+// TWOFISH
+void TwoFishBlockEncryptor::Initialize(TArray<byte>* InKey)
+{
+	Key = InKey;
+
+	if (Key->Num() != 16 && Key->Num() != 24 && Key->Num() != 32)
+	{
+		LowLevelFatalError(TEXT("Incorrect key size. %i size chosen"), Key->Num());
+	}
+
+	Encryptor = CryptoPP::Twofish::Encryption(Key->GetData(), Key->Num());
+	Decryptor = CryptoPP::Twofish::Decryption(Key->GetData(), Key->Num());
+
+	FixedBlockSize = 16;
+}
+
+void TwoFishBlockEncryptor::EncryptBlock(byte* Block)
+{
+	byte* Output = new byte[FixedBlockSize];
+	Encryptor.ProcessBlock(Block, Output);
+	memcpy(Block, Output, FixedBlockSize);
+	delete[] Output;
+
+	//UE_LOG(PacketHandlerLog, Log, TEXT("TwoFish Encrypted"));
+}
+
+void TwoFishBlockEncryptor::DecryptBlock(byte* Block)
+{
+	byte* Output = new byte[FixedBlockSize];
+	Decryptor.ProcessBlock(Block, Output);
+	memcpy(Block, Output, FixedBlockSize);
+	delete[] Output;
+
+	//UE_LOG(PacketHandlerLog, Log, TEXT("TwoFish Decrypted"));
+}
diff --git a/Engine/Source/Runtime/PacketHandlers/EncryptionComponents/SymmetricEncryption/BlockEncryption/TwoFishBlockEncryptor/Public/TwoFishBlockEncryptor.h b/Engine/Source/Runtime/PacketHandlers/EncryptionComponents/SymmetricEncryption/BlockEncryption/TwoFishBlockEncryptor/Public/TwoFishBlockEncryptor.h
new file mode 100644
--- /dev/null
+++ b/Engine/Source/Runtime/PacketHandlers/EncryptionComponents/SymmetricEncryption/BlockEncryption/TwoFishBlockEncryptor/Public/TwoFishBlockEncryptor.h
@@ -0,0 +1,36 @@
+// Copyright 1998-2015 Epic Games, Inc. All Rights Reserved.
+
+#pragma once
+
+#include "BlockEncryptionHandlerComponent.h"
+#include "CryptoPP/5.6.2/include/twofish.h"
+
+/* TwoFish Block Encryptor Module Interface */
+class FTwoFishBlockEncryptorModuleInterface : public FBlockEncryptorModuleInterface
+{
+	virtual BlockEncryptor* CreateBlockEncryptorInstance() override;
+};
+
+/*
+* TwoFish Block encryption
+*/
+class TWOFISHBLOCKENCRYPTOR_API TwoFishBlockEncryptor : public BlockEncryptor
+{
+public:
+	/* Initialized the encryptor */
+	void Initialize(TArray<byte>* Key) override;
+
+	/* Encrypts outgoing packets */
+	void EncryptBlock(byte* Block) override;
+
+	/* Decrypts incoming packets */
+	void DecryptBlock(byte* Block) override;
+
+	/* Get the default key size for this encryptor */
+	uint32 GetDefaultKeySize() { return 16; }
+
+private:
+	/* Encryptors for AES */
+	CryptoPP::Twofish::Encryption Encryptor;
+	CryptoPP::Twofish::Decryption Decryptor;
+};
\ No newline at end of file
diff --git a/Engine/Source/Runtime/PacketHandlers/EncryptionComponents/SymmetricEncryption/BlockEncryption/TwoFishBlockEncryptor/TwoFishBlockEncryptor.Build.cs b/Engine/Source/Runtime/PacketHandlers/EncryptionComponents/SymmetricEncryption/BlockEncryption/TwoFishBlockEncryptor/TwoFishBlockEncryptor.Build.cs
new file mode 100644
--- /dev/null
+++ b/Engine/Source/Runtime/PacketHandlers/EncryptionComponents/SymmetricEncryption/BlockEncryption/TwoFishBlockEncryptor/TwoFishBlockEncryptor.Build.cs
@@ -0,0 +1,27 @@
+// Copyright 1998-2015 Epic Games, Inc. All Rights Reserved.
+
+using UnrealBuildTool;
+using System.IO;
+
+public class TwoFishBlockEncryptor : ModuleRules
+{
+    public TwoFishBlockEncryptor(TargetInfo Target)
+    {
+        PublicDependencyModuleNames.AddRange(
+            new string[] {
+				"Core",
+                "BlockEncryptionHandlerComponent"
+            }
+        );
+
+        if ((Target.Platform == UnrealTargetPlatform.Win64) ||
+            (Target.Platform == UnrealTargetPlatform.Win32))
+        {
+            AddThirdPartyPrivateStaticDependencies(Target,
+                "CryptoPP"
+                );
+        }
+
+        PublicIncludePathModuleNames.Add("CryptoPP");
+    }
+}
diff --git a/Engine/Source/Runtime/PacketHandlers/EncryptionComponents/SymmetricEncryption/BlockEncryption/XORBlockEncryptor/Private/XORBlockEncryptor.cpp b/Engine/Source/Runtime/PacketHandlers/EncryptionComponents/SymmetricEncryption/BlockEncryption/XORBlockEncryptor/Private/XORBlockEncryptor.cpp
new file mode 100644
--- /dev/null
+++ b/Engine/Source/Runtime/PacketHandlers/EncryptionComponents/SymmetricEncryption/BlockEncryption/XORBlockEncryptor/Private/XORBlockEncryptor.cpp
@@ -0,0 +1,76 @@
+// Copyright 1998-2015 Epic Games, Inc. All Rights Reserved.
+
+#include "XORBlockEncryptor.h"
+
+IMPLEMENT_MODULE(FXORBlockEncryptorModuleInterface, XORBlockEncryptor);
+
+// MODULE INTERFACE
+BlockEncryptor* FXORBlockEncryptorModuleInterface::CreateBlockEncryptorInstance()
+{
+	return new XORBlockEncryptor;
+}
+
+// XOR BLOCK
+void XORBlockEncryptor::Initialize(TArray<byte>* InKey)
+{
+	Key = InKey;
+
+	if (Key->Num() != sizeof(int8) && Key->Num() != sizeof(int16) && Key->Num() != sizeof(int32) && Key->Num() != sizeof(int64))
+	{
+		LowLevelFatalError(TEXT("Incorrect key size. %i size chosen"), Key->Num());
+	}
+
+	FixedBlockSize = Key->Num();
+}
+
+void XORBlockEncryptor::EncryptBlock(byte* Block)
+{
+	if (Key->Num() == sizeof(int8))
+	{
+		int8 Output = *reinterpret_cast<int8*>(Block) ^ *reinterpret_cast<int8*>(Key->GetData());
+		memcpy(Block, &Output, sizeof(int8));
+	} 
+	else if (Key->Num() == sizeof(int16))
+	{
+		int16 Output = *reinterpret_cast<int16*>(Block) ^ *reinterpret_cast<int16*>(Key->GetData());
+		memcpy(Block, &Output, sizeof(int16));
+	} 
+	else if (Key->Num() == sizeof(int32))
+	{
+		int32 Output = *reinterpret_cast<int32*>(Block) ^ *reinterpret_cast<int32*>(Key->GetData());
+		memcpy(Block, &Output, sizeof(int32));
+	} 
+	else if (Key->Num() == sizeof(int64))
+	{
+		int64 Output = *reinterpret_cast<int64*>(Block) ^ *reinterpret_cast<int64*>(Key->GetData());
+		memcpy(Block, &Output, sizeof(int64));
+	}
+
+	//UE_LOG(PacketHandlerLog, Log, TEXT("XOR Block Encrypted"));
+}
+
+void XORBlockEncryptor::DecryptBlock(byte* Block)
+{
+	if (Key->Num() == sizeof(int8))
+	{
+		int8 Output = *reinterpret_cast<int8*>(Block) ^ *reinterpret_cast<int8*>(Key->GetData());
+		memcpy(Block, &Output, sizeof(int8));
+	} 
+	else if (Key->Num() == sizeof(int16))
+	{
+		int16 Output = *reinterpret_cast<int16*>(Block) ^ *reinterpret_cast<int16*>(Key->GetData());
+		memcpy(Block, &Output, sizeof(int16));
+	} 
+	else if (Key->Num() == sizeof(int32))
+	{
+		int32 Output = *reinterpret_cast<int32*>(Block) ^ *reinterpret_cast<int32*>(Key->GetData());
+		memcpy(Block, &Output, sizeof(int32));
+	} 
+	else if (Key->Num() == sizeof(int64))
+	{
+		int64 Output = *reinterpret_cast<int64*>(Block) ^ *reinterpret_cast<int64*>(Key->GetData());
+		memcpy(Block, &Output, sizeof(int64));
+	}
+
+	//UE_LOG(PacketHandlerLog, Log, TEXT("XOR Block Decrypted"));
+}
\ No newline at end of file
diff --git a/Engine/Source/Runtime/PacketHandlers/EncryptionComponents/SymmetricEncryption/BlockEncryption/XORBlockEncryptor/Public/XORBlockEncryptor.h b/Engine/Source/Runtime/PacketHandlers/EncryptionComponents/SymmetricEncryption/BlockEncryption/XORBlockEncryptor/Public/XORBlockEncryptor.h
new file mode 100644
--- /dev/null
+++ b/Engine/Source/Runtime/PacketHandlers/EncryptionComponents/SymmetricEncryption/BlockEncryption/XORBlockEncryptor/Public/XORBlockEncryptor.h
@@ -0,0 +1,30 @@
+// Copyright 1998-2015 Epic Games, Inc. All Rights Reserved.
+
+#pragma once
+
+#include "BlockEncryptionHandlerComponent.h"
+
+/* XOR Block Encryptor Module Interface */
+class FXORBlockEncryptorModuleInterface : public FBlockEncryptorModuleInterface
+{
+	virtual BlockEncryptor* CreateBlockEncryptorInstance() override;
+};
+
+/*
+* XOR Block encryption
+*/
+class XORBLOCKENCRYPTOR_API XORBlockEncryptor : public BlockEncryptor
+{
+public:
+	/* Initialized the encryptor */
+	void Initialize(TArray<byte>* Key) override;
+
+	/* Encrypts outgoing packets */
+	void EncryptBlock(byte* Block) override;
+
+	/* Decrypts incoming packets */
+	void DecryptBlock(byte* Block) override;
+
+	/* Get the default key size for this encryptor */
+	uint32 GetDefaultKeySize() { return 4; }
+};
diff --git a/Engine/Source/Runtime/PacketHandlers/EncryptionComponents/SymmetricEncryption/BlockEncryption/XORBlockEncryptor/XORBlockEncryptor.Build.cs b/Engine/Source/Runtime/PacketHandlers/EncryptionComponents/SymmetricEncryption/BlockEncryption/XORBlockEncryptor/XORBlockEncryptor.Build.cs
new file mode 100644
--- /dev/null
+++ b/Engine/Source/Runtime/PacketHandlers/EncryptionComponents/SymmetricEncryption/BlockEncryption/XORBlockEncryptor/XORBlockEncryptor.Build.cs
@@ -0,0 +1,27 @@
+// Copyright 1998-2015 Epic Games, Inc. All Rights Reserved.
+
+using UnrealBuildTool;
+using System.IO;
+
+public class XORBlockEncryptor : ModuleRules
+{
+    public XORBlockEncryptor(TargetInfo Target)
+    {
+        PublicDependencyModuleNames.AddRange(
+            new string[] {
+				"Core",
+                "BlockEncryptionHandlerComponent"
+            }
+        );
+
+        if ((Target.Platform == UnrealTargetPlatform.Win64) ||
+        (Target.Platform == UnrealTargetPlatform.Win32))
+        {
+            AddThirdPartyPrivateStaticDependencies(Target,
+                "CryptoPP"
+                );
+        }
+
+        PublicIncludePathModuleNames.Add("CryptoPP");
+    }
+}
diff --git a/Engine/Source/Runtime/PacketHandlers/EncryptionComponents/SymmetricEncryption/StreamEncryption/StreamEncryptionHandlerComponent/Private/StreamEncryptionHandlerComponent.cpp b/Engine/Source/Runtime/PacketHandlers/EncryptionComponents/SymmetricEncryption/StreamEncryption/StreamEncryptionHandlerComponent/Private/StreamEncryptionHandlerComponent.cpp
new file mode 100644
--- /dev/null
+++ b/Engine/Source/Runtime/PacketHandlers/EncryptionComponents/SymmetricEncryption/StreamEncryption/StreamEncryptionHandlerComponent/Private/StreamEncryptionHandlerComponent.cpp
@@ -0,0 +1,154 @@
+// Copyright 1998-2015 Epic Games, Inc. All Rights Reserved.
+
+#include "StreamEncryptionHandlerComponent.h"
+#include "XORStreamEncryptor.h"
+
+#include "CryptoPP/5.6.2/include/osrng.h"
+
+IMPLEMENT_MODULE(FStreamEncryptionHandlerComponentModuleInterface, StreamEncryptionHandlerComponent);
+
+///////////////////////////////////////////////////
+// STREAM ENCRYPTION
+///////////////////////////////////////////////////
+
+// Stream Encryption Handler
+StreamEncryptionHandlerComponent::StreamEncryptionHandlerComponent(StreamEncryptor* InEncryptor, uint32 InKeySizeInBytes)
+{
+	if (InEncryptor != nullptr)
+	{
+		Encryptor = InEncryptor;
+	} else
+	{
+		Encryptor = new XORStreamEncryptor;
+	}
+
+	// Set key size, if no key size passed in then get default key size for the encryptor
+	KeySizeInBytes = InKeySizeInBytes != 0 ? InKeySizeInBytes : Encryptor->GetDefaultKeySize();
+}
+
+void StreamEncryptionHandlerComponent::Initialize()
+{
+	if (Handler->Mode == Handler::Mode::Client)
+	{
+		// Generate the key
+		CryptoPP::AutoSeededRandomPool Rng;
+
+		Key.Reserve(KeySizeInBytes);
+		uint8* Secret = new uint8[KeySizeInBytes];
+
+		Rng.GenerateBlock(Secret, KeySizeInBytes);
+
+		for (uint32 i = 0; i < KeySizeInBytes; ++i)
+		{
+			Key.Add(Secret[i]);
+		}
+
+		Encryptor->Initialize(&Key);
+		SetState(Handler::Component::State::InitializedOnLocal);
+	}
+
+	SetActive(true);
+}
+
+bool StreamEncryptionHandlerComponent::IsValid() const
+{
+	return Key.Num() > 0;
+}
+
+void StreamEncryptionHandlerComponent::Incoming(FBitReader& Packet)
+{
+	switch (State)
+	{
+		case Handler::Component::State::InitializedOnLocal:
+		case Handler::Component::State::UnInitialized:
+		{
+			TArray<uint8> ReceivedKey;
+			Packet << ReceivedKey;
+
+			// Set key
+			Key = ReceivedKey;
+			KeySizeInBytes = uint32(Key.Num());
+
+			Encryptor->Initialize(&Key);
+			SetState(Handler::Component::State::Initialized);
+			Initialized();
+			break;
+		}
+		case Handler::Component::State::Initialized:
+		{
+			// Handle this packet
+			if (IsValid() && Packet.GetNumBytes() > 0)
+			{
+				DecryptStream(Packet);
+			}
+			break;
+		}
+		default:
+		{
+			break;
+		}
+	}
+}
+
+void StreamEncryptionHandlerComponent::Outgoing(FBitWriter& Packet)
+{
+	switch (State)
+	{
+		case Handler::Component::State::UnInitialized:
+		{
+			break;
+		}
+		// Send packet initializing remote handler
+		case Handler::Component::State::InitializedOnLocal:
+		{
+			FBitWriter Local;
+			Local.AllowAppend(true);
+			Local.SetAllowResize(true);
+
+			Local << Key;
+			Local.Serialize(Packet.GetData(), Packet.GetNumBytes());
+			Packet = Local;
+
+			SetState(Handler::Component::State::Initialized);
+			SetActive(true);
+			Initialized();
+			break;
+		}
+		case Handler::Component::State::Initialized:
+		{
+			// Handle this packet
+			if (IsValid() && Packet.GetNumBytes() > 0)
+			{
+				EncryptStream(Packet);
+			}
+			break;
+		}
+		default:
+		{
+			break;
+		}
+	}
+}
+
+void StreamEncryptionHandlerComponent::EncryptStream(FBitWriter& Packet)
+{
+	Encryptor->EncryptStream(Packet.GetData(), Packet.GetNumBytes());
+}
+
+void StreamEncryptionHandlerComponent::DecryptStream(FBitReader& Packet)
+{
+	Encryptor->DecryptStream(Packet.GetData(), Packet.GetNumBytes());
+}
+
+// MODULE INTERFACE
+HandlerComponent* FStreamEncryptionHandlerComponentModuleInterface::CreateComponentInstance()
+{
+	return new StreamEncryptionHandlerComponent;
+}
+
+HandlerComponent* FStreamEncryptionHandlerComponentModuleInterface::CreateComponentInstance(FString& Options)
+{
+	TSharedPtr<IModuleInterface> Interface = FModuleManager::Get().LoadModule(FName(*Options));
+	TSharedPtr<FStreamEncryptorModuleInterface> StreamEncryptorInterface(static_cast<FStreamEncryptorModuleInterface*>(&(*Interface)));
+	return new StreamEncryptionHandlerComponent(StreamEncryptorInterface->CreateStreamEncryptorInstance());
+}
\ No newline at end of file
diff --git a/Engine/Source/Runtime/PacketHandlers/EncryptionComponents/SymmetricEncryption/StreamEncryption/StreamEncryptionHandlerComponent/Public/StreamEncryptionHandlerComponent.h b/Engine/Source/Runtime/PacketHandlers/EncryptionComponents/SymmetricEncryption/StreamEncryption/StreamEncryptionHandlerComponent/Public/StreamEncryptionHandlerComponent.h
new file mode 100644
--- /dev/null
+++ b/Engine/Source/Runtime/PacketHandlers/EncryptionComponents/SymmetricEncryption/StreamEncryption/StreamEncryptionHandlerComponent/Public/StreamEncryptionHandlerComponent.h
@@ -0,0 +1,77 @@
+// Copyright 1998-2015 Epic Games, Inc. All Rights Reserved.
+
+#pragma once
+
+#include "PacketHandler.h"
+
+class STREAMENCRYPTIONHANDLERCOMPONENT_API StreamEncryptor
+{
+public:
+	/* Initialized the encryptor */
+	virtual void Initialize(TArray<uint8>* Key) PURE_VIRTUAL(StreamEncryptor::Initialize, );
+
+	/* Encrypts block */
+	virtual void EncryptStream(uint8* Stream, uint32 BytesCount) PURE_VIRTUAL(StreamEncryptor::EncryptBlock, );
+
+	/* Decrypts incoming packets */
+	virtual void DecryptStream(uint8* Stream, uint32 BytesCount) PURE_VIRTUAL(StreamEncryptor::DecryptBlock, );
+
+	/* Gets the default key size for this encryptor */
+	virtual uint32 GetDefaultKeySize() PURE_VIRTUAL(StreamEncryptor::GetDefaultKeySize, return 0;);
+
+protected:
+	/* Key, this is initialized when Initialize() is called */
+	TArray<uint8>* Key;
+};
+
+/* Stream Encryptor Module Interface */
+class FStreamEncryptorModuleInterface : public IModuleInterface
+{
+public:
+	virtual StreamEncryptor* CreateStreamEncryptorInstance() = 0;
+};
+
+/*
+* Symmetric block cipher handler component using AES as the cipher
+*/
+class STREAMENCRYPTIONHANDLERCOMPONENT_API StreamEncryptionHandlerComponent : public HandlerComponent
+{
+public:
+	/* Initializes default data, can provide the key size in bytes */
+	StreamEncryptionHandlerComponent(StreamEncryptor* Encryptor = nullptr, uint32 KeySizeInBytes = 0);
+
+	/* Initializes the handler component */
+	virtual void Initialize() override;
+
+	/* Whether the handler component is valid */
+	virtual bool IsValid() const override;
+
+	/* Handles any incoming packets */
+	virtual void Incoming(FBitReader& Packet) override;
+
+	/* Handles any outgoing packets */
+	virtual void Outgoing(FBitWriter& Packet) override;
+
+protected:
+	void EncryptStream(FBitWriter& Packet);
+	void DecryptStream(FBitReader& Packet);
+
+	/* The encryptor used */
+	StreamEncryptor* Encryptor;
+
+	/** Size of Key in Bits */
+	uint32 KeySizeInBytes;
+
+	// Key used for symmetric encryption
+	TArray<uint8> Key;
+};
+
+/* Stream Encryption Handler Component Module Interface */
+class FStreamEncryptionHandlerComponentModuleInterface : public FPacketHandlerComponentModuleInterface
+{
+public:
+	/* Creates an instance of this component */
+	virtual HandlerComponent* CreateComponentInstance() override;
+	/* Creates an instance of this component with an FString of options */
+	virtual HandlerComponent* CreateComponentInstance(FString& Options) override;
+};
diff --git a/Engine/Source/Runtime/PacketHandlers/EncryptionComponents/SymmetricEncryption/StreamEncryption/StreamEncryptionHandlerComponent/StreamEncryptionHandlerComponent.Build.cs b/Engine/Source/Runtime/PacketHandlers/EncryptionComponents/SymmetricEncryption/StreamEncryption/StreamEncryptionHandlerComponent/StreamEncryptionHandlerComponent.Build.cs
new file mode 100644
--- /dev/null
+++ b/Engine/Source/Runtime/PacketHandlers/EncryptionComponents/SymmetricEncryption/StreamEncryption/StreamEncryptionHandlerComponent/StreamEncryptionHandlerComponent.Build.cs
@@ -0,0 +1,30 @@
+// Copyright 1998-2015 Epic Games, Inc. All Rights Reserved.
+
+using UnrealBuildTool;
+using System.IO;
+
+public class StreamEncryptionHandlerComponent : ModuleRules
+{
+    public StreamEncryptionHandlerComponent(TargetInfo Target)
+    {
+        PublicDependencyModuleNames.AddRange(
+            new string[] {
+				"Core",
+                "PacketHandler",
+                "XORStreamEncryptor"
+            }
+        );
+
+        CircularlyReferencedDependentModules.Add("XORStreamEncryptor");
+
+        if ((Target.Platform == UnrealTargetPlatform.Win64) ||
+         (Target.Platform == UnrealTargetPlatform.Win32))
+        {
+            AddThirdPartyPrivateStaticDependencies(Target,
+                "CryptoPP"
+                );
+        }
+
+        PublicIncludePathModuleNames.Add("CryptoPP");
+    }
+}
diff --git a/Engine/Source/Runtime/PacketHandlers/EncryptionComponents/SymmetricEncryption/StreamEncryption/XORStreamEncryptor/Private/XORStreamEncryptor.cpp b/Engine/Source/Runtime/PacketHandlers/EncryptionComponents/SymmetricEncryption/StreamEncryption/XORStreamEncryptor/Private/XORStreamEncryptor.cpp
new file mode 100644
--- /dev/null
+++ b/Engine/Source/Runtime/PacketHandlers/EncryptionComponents/SymmetricEncryption/StreamEncryption/XORStreamEncryptor/Private/XORStreamEncryptor.cpp
@@ -0,0 +1,38 @@
+// Copyright 1998-2015 Epic Games, Inc. All Rights Reserved.
+
+#include "XORStreamEncryptor.h"
+
+IMPLEMENT_MODULE(FXORStreamEncryptorModuleInterface, XORStreamEncryptor);
+
+StreamEncryptor* FXORStreamEncryptorModuleInterface::CreateStreamEncryptorInstance()
+{
+	return new XORStreamEncryptor;
+}
+
+// XOR Stream
+void XORStreamEncryptor::Initialize(TArray<uint8>* InKey)
+{
+	Key = InKey;
+}
+
+void XORStreamEncryptor::EncryptStream(uint8* Block, uint32 BytesCount)
+{
+	// XOR
+	for (uint32 i = 0; i < BytesCount; ++i)
+	{
+		Block[i] ^= (*Key)[i % (Key->Num() - 1)];
+	}
+
+	//UE_LOG(PacketHandlerLog, Log, TEXT("XOR Stream Encrypted"));
+}
+
+void XORStreamEncryptor::DecryptStream(uint8* Block, uint32 BytesCount)
+{
+	// XOR
+	for (uint32 i = 0; i < BytesCount; ++i)
+	{
+		Block[i] ^= (*Key)[i % (Key->Num() - 1)];
+	}
+
+	//UE_LOG(PacketHandlerLog, Log, TEXT("XOR Stream Decrypted"));
+}
diff --git a/Engine/Source/Runtime/PacketHandlers/EncryptionComponents/SymmetricEncryption/StreamEncryption/XORStreamEncryptor/Public/XORStreamEncryptor.h b/Engine/Source/Runtime/PacketHandlers/EncryptionComponents/SymmetricEncryption/StreamEncryption/XORStreamEncryptor/Public/XORStreamEncryptor.h
new file mode 100644
--- /dev/null
+++ b/Engine/Source/Runtime/PacketHandlers/EncryptionComponents/SymmetricEncryption/StreamEncryption/XORStreamEncryptor/Public/XORStreamEncryptor.h
@@ -0,0 +1,27 @@
+#include "StreamEncryptionHandlerComponent.h"
+
+/* XOR Stream Encryptor Module Interface */
+class FXORStreamEncryptorModuleInterface : public FStreamEncryptorModuleInterface
+{
+public:
+	virtual StreamEncryptor* CreateStreamEncryptorInstance();
+};
+
+/*
+* XOR Block encryption
+*/
+class XORSTREAMENCRYPTOR_API XORStreamEncryptor : public StreamEncryptor
+{
+public:
+	/* Initialized the encryptor */
+	void Initialize(TArray<uint8>* Key) override;
+
+	/* Encrypts outgoing packets */
+	void EncryptStream(uint8* Stream, uint32 BytesCount) override;
+
+	/* Decrypts incoming packets */
+	void DecryptStream(uint8* Stream, uint32 BytesCount) override;
+
+	/* Get the default key size for this encryptor */
+	uint32 GetDefaultKeySize() { return 4; }
+};
\ No newline at end of file
diff --git a/Engine/Source/Runtime/PacketHandlers/EncryptionComponents/SymmetricEncryption/StreamEncryption/XORStreamEncryptor/XORStreamEncryptor.Build.cs b/Engine/Source/Runtime/PacketHandlers/EncryptionComponents/SymmetricEncryption/StreamEncryption/XORStreamEncryptor/XORStreamEncryptor.Build.cs
new file mode 100644
--- /dev/null
+++ b/Engine/Source/Runtime/PacketHandlers/EncryptionComponents/SymmetricEncryption/StreamEncryption/XORStreamEncryptor/XORStreamEncryptor.Build.cs
@@ -0,0 +1,17 @@
+// Copyright 1998-2015 Epic Games, Inc. All Rights Reserved.
+
+using UnrealBuildTool;
+using System.IO;
+
+public class XORStreamEncryptor : ModuleRules
+{
+    public XORStreamEncryptor(TargetInfo Target)
+    {
+        PublicDependencyModuleNames.AddRange(
+            new string[] {
+				"Core",
+                "StreamEncryptionHandlerComponent"
+            }
+        );
+    }
+}
diff --git a/Engine/Source/Runtime/PacketHandlers/PacketHandler/PacketHandler.Build.cs b/Engine/Source/Runtime/PacketHandlers/PacketHandler/PacketHandler.Build.cs
--- a/Engine/Source/Runtime/PacketHandlers/PacketHandler/PacketHandler.Build.cs
+++ b/Engine/Source/Runtime/PacketHandlers/PacketHandler/PacketHandler.Build.cs
@@ -15,7 +15,5 @@ public class PacketHandler : ModuleRules
         );
 
         CircularlyReferencedDependentModules.Add("ReliabilityHandlerComponent");
-
-        //DynamicallyLoadedModuleNames.Add("EncryptionHandlerComponent");
     }
 }
diff --git a/Engine/Source/Runtime/PacketHandlers/PacketHandler/Private/PacketHandler.cpp b/Engine/Source/Runtime/PacketHandlers/PacketHandler/Private/PacketHandler.cpp
--- a/Engine/Source/Runtime/PacketHandlers/PacketHandler/Private/PacketHandler.cpp
+++ b/Engine/Source/Runtime/PacketHandlers/PacketHandler/Private/PacketHandler.cpp
@@ -1,7 +1,7 @@
 #include "PacketHandler.h"
 #include "ReliabilityHandlerComponent.h"
 
-IMPLEMENT_MODULE(FPacketHandlerModuleInterface, PacketHandler);
+IMPLEMENT_MODULE(FPacketHandlerComponentModuleInterface, PacketHandler);
 
 DEFINE_LOG_CATEGORY(PacketHandlerLog);
 
@@ -56,27 +56,101 @@ void PacketHandler::Tick(float DeltaTime)
 
 void PacketHandler::Initialize(Handler::Mode InMode)
 {
+	struct ComponentAndOptions
+	{
+		FString ComponentName;
+		FString Options;
+	};
+
 	Mode = InMode;
 
 	FString Components;
 	GConfig->GetString(TEXT("PacketHandlerComponents"), TEXT("Components"), Components, GEngineIni);
 
-	TArray<FString> ComponentsArray;
-	Components.ParseIntoArray(ComponentsArray, TEXT(","), true);
+	TArray<ComponentAndOptions> ComponentsArray;
+	ComponentsArray.Add(ComponentAndOptions());
+
+	int32 i = 0;
+	int32 ComponentsCount = 0;
+	while(i < Components.Len())
+	{
+		TCHAR c = Components[i];
+
+		// End of string
+		if(c == ',')
+		{
+			// If not end of string add a new string
+			if(i + 1 < Components.Len())
+			{
+				ComponentsArray.Add(ComponentAndOptions());
+				++ComponentsCount;
+			}
+		}
+		// Parsing Options
+		else if(c == '(')
+		{
+			// Skip '('
+			++i;
+
+			// Parse until end of options
+			while(i < Components.Len())
+			{
+				c = Components[i];
+
+				// End of options
+				if(c == ')')
+				{
+					break;
+				}
+				// Append char to options
+				else
+				{
+					ComponentsArray[ComponentsCount].Options.AppendChar(c);
+				}	
+
+				++i;
+			}
+		}
+		// Append char to component name if not whitespace
+		else if(c != ' ' )
+		{
+			ComponentsArray[ComponentsCount].ComponentName.AppendChar(c);
+		}
+
+		++i;
+	}
 
 	for (int32 i = 0; i < ComponentsArray.Num(); i++)
 	{
-		if(ComponentsArray[i] == TEXT("ReliabilityHandlerComponent"))
+		// Skip adding reliability component as it is added automatically in the handler
+		if(ComponentsArray[i].ComponentName == TEXT("ReliabilityHandlerComponent"))
 		{
 			continue;
 		}
 
-		TSharedPtr<IModuleInterface> Interface = FModuleManager::Get().LoadModule(FName(*ComponentsArray[i]));
+		TSharedPtr<IModuleInterface> Interface = FModuleManager::Get().LoadModule(FName(*ComponentsArray[i].ComponentName));
 
 		if(Interface.IsValid())
 		{
-			TSharedPtr<FPacketHandlerModuleInterface> PacketHandlerInterface(static_cast<FPacketHandlerModuleInterface*>(&(*Interface)));
-			Add(PacketHandlerInterface->CreateComponentInstance());
+			TSharedPtr<FPacketHandlerComponentModuleInterface> PacketHandlerInterface(static_cast<FPacketHandlerComponentModuleInterface*>(&(*Interface)));
+
+			if(PacketHandlerInterface.IsValid())
+			{
+				// Options Specified
+				if(ComponentsArray[i].Options.Len() > 0)
+				{
+					Add(PacketHandlerInterface->CreateComponentInstance(ComponentsArray[i].Options));
+				}
+				// No Options Specified
+				else
+				{
+					Add(PacketHandlerInterface->CreateComponentInstance());
+				}
+			}
+		}
+		else
+		{
+			UE_LOG(PacketHandlerLog, Warning, TEXT("Unable to Load Module: %s"), *ComponentsArray[i].ComponentName);
 		}
 	}
 }
diff --git a/Engine/Source/Runtime/PacketHandlers/PacketHandler/Public/PacketHandler.h b/Engine/Source/Runtime/PacketHandlers/PacketHandler/Public/PacketHandler.h
--- a/Engine/Source/Runtime/PacketHandlers/PacketHandler/Public/PacketHandler.h
+++ b/Engine/Source/Runtime/PacketHandlers/PacketHandler/Public/PacketHandler.h
@@ -207,9 +207,11 @@ private:
 };
 
 /** Packet Handler Module Interface */
-class FPacketHandlerModuleInterface : public IModuleInterface
+class FPacketHandlerComponentModuleInterface : public IModuleInterface
 {
 public:
+	/* Creates an instance of this component */
 	virtual HandlerComponent* CreateComponentInstance() PURE_VIRTUAL(FPacketHandlerModuleInterface::CreateComponentInstance,return nullptr;);
-	virtual HandlerComponent* CreateComponentInstance(FName& Options) PURE_VIRTUAL(FPacketHandlerModuleInterface::CreateComponentInstance,return nullptr;);
+	/* Creates an instance of this component with an FString of options */
+	virtual HandlerComponent* CreateComponentInstance(FString& Options) PURE_VIRTUAL(FPacketHandlerModuleInterface::CreateComponentInstance,return nullptr;);
 };
diff --git a/Engine/Source/Runtime/PacketHandlers/ReliabilityHandlerComponent/Public/ReliabilityHandlerComponent.h b/Engine/Source/Runtime/PacketHandlers/ReliabilityHandlerComponent/Public/ReliabilityHandlerComponent.h
--- a/Engine/Source/Runtime/PacketHandlers/ReliabilityHandlerComponent/Public/ReliabilityHandlerComponent.h
+++ b/Engine/Source/Runtime/PacketHandlers/ReliabilityHandlerComponent/Public/ReliabilityHandlerComponent.h
@@ -52,7 +52,7 @@ protected:
 };
 
 /* Reliability Module Interface */
-class FReliabilityHandlerComponentModuleInterface : public FPacketHandlerModuleInterface
+class FReliabilityHandlerComponentModuleInterface : public FPacketHandlerComponentModuleInterface
 {
 public:
 	virtual HandlerComponent* CreateComponentInstance() override;
