Fixing variable shadow

diff --git a/Engine/Source/Runtime/PacketHandlers/PacketHandler/Private/PacketHandler.cpp b/Engine/Source/Runtime/PacketHandlers/PacketHandler/Private/PacketHandler.cpp
--- a/Engine/Source/Runtime/PacketHandlers/PacketHandler/Private/PacketHandler.cpp
+++ b/Engine/Source/Runtime/PacketHandlers/PacketHandler/Private/PacketHandler.cpp
@@ -120,7 +120,7 @@ void PacketHandler::Initialize(Handler::Mode InMode)
 		++i;
 	}
 
-	for (int32 i = 0; i < ComponentsArray.Num(); i++)
+	for (i = 0; i < ComponentsArray.Num(); i++)
 	{
 		// Skip adding reliability component as it is added automatically in the handler
 		if(ComponentsArray[i].ComponentName == TEXT("ReliabilityHandlerComponent"))
