Updated/refactored PacketHandler code, and gave UNetConnections a PacketHandler by default.

- Split [PacketHandlerComponents].Components ini setting, into an array, instead of one string.
- Give UNetConnection's a PacketHandler by default (for future handshake code).
- Refactored/cleaned-up a lot of PacketHandler code.

diff --git a/Engine/Source/Runtime/Engine/Private/NetConnection.cpp b/Engine/Source/Runtime/Engine/Private/NetConnection.cpp
--- a/Engine/Source/Runtime/Engine/Private/NetConnection.cpp
+++ b/Engine/Source/Runtime/Engine/Private/NetConnection.cpp
@@ -83,6 +83,17 @@ void UNetConnection::InitBase(UNetDriver* InDriver,class FSocket* InSocket, cons
 	// Owning net driver
 	Driver = InDriver;
 
+	// Reset Handler
+	Handler.Reset(NULL);
+
+	Handler = MakeUnique<PacketHandler>();
+
+	if(Handler.IsValid())
+	{
+		Handler::Mode Mode = Driver->ServerConnection != nullptr ? Handler::Mode::Client : Handler::Mode::Server;
+		Handler->Initialize(Mode);
+	}
+
 	// Stats
 	StatUpdateTime			= Driver->Time;
 	LastReceiveTime			= Driver->Time;
@@ -141,6 +152,12 @@ void UNetConnection::InitBase(UNetDriver* InDriver,class FSocket* InSocket, cons
 void UNetConnection::InitConnection(UNetDriver* InDriver, EConnectionState InState, const FURL& InURL, int32 InConnectionSpeed, int32 InMaxPacket)
 {
 	Driver = InDriver;
+
+	if (!Handler.IsValid())
+	{
+		Handler = MakeUnique<PacketHandler>();
+	}
+
 	// We won't be sending any packets, so use a default size
 	MaxPacket = (InMaxPacket == 0 || InMaxPacket > MAX_PACKET_SIZE) ? MAX_PACKET_SIZE : InMaxPacket;
 	PacketOverhead = 0;
@@ -303,6 +320,8 @@ void UNetConnection::CleanUp()
 
 	CleanupDormantActorState();
 
+	Handler.Reset(NULL);
+
 	Driver = NULL;
 }
 
diff --git a/Engine/Source/Runtime/Online/OnlineSubsystemUtils/Private/IpConnection.cpp b/Engine/Source/Runtime/Online/OnlineSubsystemUtils/Private/IpConnection.cpp
--- a/Engine/Source/Runtime/Online/OnlineSubsystemUtils/Private/IpConnection.cpp
+++ b/Engine/Source/Runtime/Online/OnlineSubsystemUtils/Private/IpConnection.cpp
@@ -42,14 +42,6 @@ void UIpConnection::InitBase(UNetDriver* InDriver, class FSocket* InSocket, cons
 
 	Socket = InSocket;
 	ResolveInfo = NULL;
-
-	// Reset Handler
-	Handler.Reset(nullptr);
-	if(Handler.IsValid())
-	{
-		Handler::Mode Mode = Driver->ServerConnection != nullptr ? Handler::Mode::Client : Handler::Mode::Server;
-		Handler->Initialize(Mode);
-	}
 }
 
 void UIpConnection::InitLocalConnection(UNetDriver* InDriver, class FSocket* InSocket, const FURL& InURL, EConnectionState InState, int32 InMaxPacket, int32 InPacketOverhead)
diff --git a/Engine/Source/Runtime/PacketHandlers/EncryptionComponents/AsymmetricEncryption/RSAEncryptionHandlerComponent/Private/RSAEncryptionHandlerComponent.cpp b/Engine/Source/Runtime/PacketHandlers/EncryptionComponents/AsymmetricEncryption/RSAEncryptionHandlerComponent/Private/RSAEncryptionHandlerComponent.cpp
--- a/Engine/Source/Runtime/PacketHandlers/EncryptionComponents/AsymmetricEncryption/RSAEncryptionHandlerComponent/Private/RSAEncryptionHandlerComponent.cpp
+++ b/Engine/Source/Runtime/PacketHandlers/EncryptionComponents/AsymmetricEncryption/RSAEncryptionHandlerComponent/Private/RSAEncryptionHandlerComponent.cpp
@@ -287,12 +287,7 @@ void RSAEncryptionHandlerComponent::DecryptWithPrivate(const TArray<byte>& InCip
 }
 
 // MODULE INTERFACE
-HandlerComponent* FRSAEncryptorHandlerComponentModuleInterface::CreateComponentInstance()
+TSharedPtr<HandlerComponent> FRSAEncryptorHandlerComponentModuleInterface::CreateComponentInstance(FString& Options)
 {
-	return new RSAEncryptionHandlerComponent;
-}
-
-HandlerComponent* FRSAEncryptorHandlerComponentModuleInterface::CreateComponentInstance(FString& Options)
-{
-	return new RSAEncryptionHandlerComponent;
+	return MakeShareable(new RSAEncryptionHandlerComponent);
 }
\ No newline at end of file
diff --git a/Engine/Source/Runtime/PacketHandlers/EncryptionComponents/AsymmetricEncryption/RSAEncryptionHandlerComponent/Public/RSAEncryptionHandlerComponent.h b/Engine/Source/Runtime/PacketHandlers/EncryptionComponents/AsymmetricEncryption/RSAEncryptionHandlerComponent/Public/RSAEncryptionHandlerComponent.h
--- a/Engine/Source/Runtime/PacketHandlers/EncryptionComponents/AsymmetricEncryption/RSAEncryptionHandlerComponent/Public/RSAEncryptionHandlerComponent.h
+++ b/Engine/Source/Runtime/PacketHandlers/EncryptionComponents/AsymmetricEncryption/RSAEncryptionHandlerComponent/Public/RSAEncryptionHandlerComponent.h
@@ -12,9 +12,7 @@ class FRSAEncryptorHandlerComponentModuleInterface : public FPacketHandlerCompon
 {
 public:
 	/* Creates an instance of this component */
-	virtual HandlerComponent* CreateComponentInstance() override;
-	/* Creates an instance of this component with an FString of options */
-	virtual HandlerComponent* CreateComponentInstance(FString& Options) override;
+	virtual TSharedPtr<HandlerComponent> CreateComponentInstance(FString& Options) override;
 };
 
 
diff --git a/Engine/Source/Runtime/PacketHandlers/EncryptionComponents/EncryptionHandlerComponent/Private/EncryptionHandlerComponent.cpp b/Engine/Source/Runtime/PacketHandlers/EncryptionComponents/EncryptionHandlerComponent/Private/EncryptionHandlerComponent.cpp
--- a/Engine/Source/Runtime/PacketHandlers/EncryptionComponents/EncryptionHandlerComponent/Private/EncryptionHandlerComponent.cpp
+++ b/Engine/Source/Runtime/PacketHandlers/EncryptionComponents/EncryptionHandlerComponent/Private/EncryptionHandlerComponent.cpp
@@ -140,12 +140,7 @@ void EncryptionHandlerComponent::Outgoing(FBitWriter& Packet)
 }
 
 //MODULE INTERFACE
-HandlerComponent* FEncryptionHandlerComponentModuleInterface::CreateComponentInstance()
+TSharedPtr<HandlerComponent> FEncryptionHandlerComponentModuleInterface::CreateComponentInstance(FString& Options)
 {
-	return new EncryptionHandlerComponent;
-}
-
-HandlerComponent* FEncryptionHandlerComponentModuleInterface::CreateComponentInstance(FString& Options)
-{
-	return new EncryptionHandlerComponent;
+	return MakeShareable(new EncryptionHandlerComponent);
 }
\ No newline at end of file
diff --git a/Engine/Source/Runtime/PacketHandlers/EncryptionComponents/EncryptionHandlerComponent/Public/EncryptionHandlerComponent.h b/Engine/Source/Runtime/PacketHandlers/EncryptionComponents/EncryptionHandlerComponent/Public/EncryptionHandlerComponent.h
--- a/Engine/Source/Runtime/PacketHandlers/EncryptionComponents/EncryptionHandlerComponent/Public/EncryptionHandlerComponent.h
+++ b/Engine/Source/Runtime/PacketHandlers/EncryptionComponents/EncryptionHandlerComponent/Public/EncryptionHandlerComponent.h
@@ -11,9 +11,7 @@ class FEncryptionHandlerComponentModuleInterface : public FPacketHandlerComponen
 {
 	public:
 	/* Creates an instance of this component */
-	virtual HandlerComponent* CreateComponentInstance() override;
-	/* Creates an instance of this component with an FString of options */
-	virtual HandlerComponent* CreateComponentInstance(FString& Options) override;
+	virtual TSharedPtr<HandlerComponent> CreateComponentInstance(FString& Options) override;
 };
 
 namespace EEncryptionHandler
diff --git a/Engine/Source/Runtime/PacketHandlers/EncryptionComponents/SymmetricEncryption/BlockEncryption/BlockEncryptionHandlerComponent/Private/BlockEncryptionHandlerComponent.cpp b/Engine/Source/Runtime/PacketHandlers/EncryptionComponents/SymmetricEncryption/BlockEncryption/BlockEncryptionHandlerComponent/Private/BlockEncryptionHandlerComponent.cpp
--- a/Engine/Source/Runtime/PacketHandlers/EncryptionComponents/SymmetricEncryption/BlockEncryption/BlockEncryptionHandlerComponent/Private/BlockEncryptionHandlerComponent.cpp
+++ b/Engine/Source/Runtime/PacketHandlers/EncryptionComponents/SymmetricEncryption/BlockEncryption/BlockEncryptionHandlerComponent/Private/BlockEncryptionHandlerComponent.cpp
@@ -222,14 +222,21 @@ uint32 BlockEncryptor::GetFixedBlockSize()
 }
 
 // MODULE INTERFACE
-HandlerComponent* FBlockEncryptionHandlerComponentModuleInterface::CreateComponentInstance()
+TSharedPtr<HandlerComponent> FBlockEncryptionHandlerComponentModuleInterface::CreateComponentInstance(FString& Options)
 {
-	return new BlockEncryptionHandlerComponent;
-}
+	TSharedPtr<HandlerComponent> ReturnVal = NULL;
 
-HandlerComponent* FBlockEncryptionHandlerComponentModuleInterface::CreateComponentInstance(FString& Options)
-{
-	TSharedPtr<IModuleInterface> Interface = FModuleManager::Get().LoadModule(FName(*Options));
-	TSharedPtr<FBlockEncryptorModuleInterface> BlockEncryptorInterface(static_cast<FBlockEncryptorModuleInterface*>(&(*Interface)));
-	return new BlockEncryptionHandlerComponent(BlockEncryptorInterface->CreateBlockEncryptorInstance());
+	if (!Options.IsEmpty())
+	{
+		TSharedPtr<IModuleInterface> Interface = FModuleManager::Get().LoadModule(FName(*Options));
+		TSharedPtr<FBlockEncryptorModuleInterface> BlockEncryptorInterface(static_cast<FBlockEncryptorModuleInterface*>(&(*Interface)));
+
+		ReturnVal = MakeShareable(new BlockEncryptionHandlerComponent(BlockEncryptorInterface->CreateBlockEncryptorInstance()));
+	}
+	else
+	{
+		ReturnVal = MakeShareable(new BlockEncryptionHandlerComponent);
+	}
+
+	return ReturnVal;
 }
\ No newline at end of file
diff --git a/Engine/Source/Runtime/PacketHandlers/EncryptionComponents/SymmetricEncryption/BlockEncryption/BlockEncryptionHandlerComponent/Public/BlockEncryptionHandlerComponent.h b/Engine/Source/Runtime/PacketHandlers/EncryptionComponents/SymmetricEncryption/BlockEncryption/BlockEncryptionHandlerComponent/Public/BlockEncryptionHandlerComponent.h
--- a/Engine/Source/Runtime/PacketHandlers/EncryptionComponents/SymmetricEncryption/BlockEncryption/BlockEncryptionHandlerComponent/Public/BlockEncryptionHandlerComponent.h
+++ b/Engine/Source/Runtime/PacketHandlers/EncryptionComponents/SymmetricEncryption/BlockEncryption/BlockEncryptionHandlerComponent/Public/BlockEncryptionHandlerComponent.h
@@ -82,7 +82,5 @@ class FBlockEncryptionHandlerComponentModuleInterface : public FPacketHandlerCom
 {
 public:
 	/* Creates an instance of this component */
-	virtual HandlerComponent* CreateComponentInstance() override;
-	/* Creates an instance of this component with an FString of options */
-	virtual HandlerComponent* CreateComponentInstance(FString& Options) override;
+	virtual TSharedPtr<HandlerComponent> CreateComponentInstance(FString& Options) override;
 };
diff --git a/Engine/Source/Runtime/PacketHandlers/EncryptionComponents/SymmetricEncryption/StreamEncryption/StreamEncryptionHandlerComponent/Private/StreamEncryptionHandlerComponent.cpp b/Engine/Source/Runtime/PacketHandlers/EncryptionComponents/SymmetricEncryption/StreamEncryption/StreamEncryptionHandlerComponent/Private/StreamEncryptionHandlerComponent.cpp
--- a/Engine/Source/Runtime/PacketHandlers/EncryptionComponents/SymmetricEncryption/StreamEncryption/StreamEncryptionHandlerComponent/Private/StreamEncryptionHandlerComponent.cpp
+++ b/Engine/Source/Runtime/PacketHandlers/EncryptionComponents/SymmetricEncryption/StreamEncryption/StreamEncryptionHandlerComponent/Private/StreamEncryptionHandlerComponent.cpp
@@ -141,14 +141,26 @@ void StreamEncryptionHandlerComponent::DecryptStream(FBitReader& Packet)
 }
 
 // MODULE INTERFACE
-HandlerComponent* FStreamEncryptionHandlerComponentModuleInterface::CreateComponentInstance()
+TSharedPtr<HandlerComponent> FStreamEncryptionHandlerComponentModuleInterface::CreateComponentInstance()
 {
-	return new StreamEncryptionHandlerComponent;
+	return MakeShareable(new StreamEncryptionHandlerComponent);
 }
 
-HandlerComponent* FStreamEncryptionHandlerComponentModuleInterface::CreateComponentInstance(FString& Options)
+TSharedPtr<HandlerComponent> FStreamEncryptionHandlerComponentModuleInterface::CreateComponentInstance(FString& Options)
 {
-	TSharedPtr<IModuleInterface> Interface = FModuleManager::Get().LoadModule(FName(*Options));
-	TSharedPtr<FStreamEncryptorModuleInterface> StreamEncryptorInterface(static_cast<FStreamEncryptorModuleInterface*>(&(*Interface)));
-	return new StreamEncryptionHandlerComponent(StreamEncryptorInterface->CreateStreamEncryptorInstance());
+	TSharedPtr<HandlerComponent> ReturnVal = NULL;
+
+	if (!Options.IsEmpty())
+	{
+		TSharedPtr<IModuleInterface> Interface = FModuleManager::Get().LoadModule(FName(*Options));
+		TSharedPtr<FStreamEncryptorModuleInterface> StreamEncryptorInterface(static_cast<FStreamEncryptorModuleInterface*>(&(*Interface)));
+
+		ReturnVal = MakeShareable(new StreamEncryptionHandlerComponent(StreamEncryptorInterface->CreateStreamEncryptorInstance()));
+	}
+	else
+	{
+		ReturnVal = MakeShareable(new StreamEncryptionHandlerComponent);
+	}
+
+	return ReturnVal;
 }
\ No newline at end of file
diff --git a/Engine/Source/Runtime/PacketHandlers/EncryptionComponents/SymmetricEncryption/StreamEncryption/StreamEncryptionHandlerComponent/Public/StreamEncryptionHandlerComponent.h b/Engine/Source/Runtime/PacketHandlers/EncryptionComponents/SymmetricEncryption/StreamEncryption/StreamEncryptionHandlerComponent/Public/StreamEncryptionHandlerComponent.h
--- a/Engine/Source/Runtime/PacketHandlers/EncryptionComponents/SymmetricEncryption/StreamEncryption/StreamEncryptionHandlerComponent/Public/StreamEncryptionHandlerComponent.h
+++ b/Engine/Source/Runtime/PacketHandlers/EncryptionComponents/SymmetricEncryption/StreamEncryption/StreamEncryptionHandlerComponent/Public/StreamEncryptionHandlerComponent.h
@@ -71,7 +71,5 @@ class FStreamEncryptionHandlerComponentModuleInterface : public FPacketHandlerCo
 {
 public:
 	/* Creates an instance of this component */
-	virtual HandlerComponent* CreateComponentInstance() override;
-	/* Creates an instance of this component with an FString of options */
-	virtual HandlerComponent* CreateComponentInstance(FString& Options) override;
+	virtual TSharedPtr<HandlerComponent> CreateComponentInstance(FString& Options) override;
 };
diff --git a/Engine/Source/Runtime/PacketHandlers/PacketHandler/Private/PacketHandler.cpp b/Engine/Source/Runtime/PacketHandlers/PacketHandler/Private/PacketHandler.cpp
--- a/Engine/Source/Runtime/PacketHandlers/PacketHandler/Private/PacketHandler.cpp
+++ b/Engine/Source/Runtime/PacketHandlers/PacketHandler/Private/PacketHandler.cpp
@@ -26,7 +26,7 @@ BufferedPacket::~BufferedPacket()
 PacketHandler::PacketHandler()
 : Time(0.f)
 , State(Handler::State::Uninitialized)
-, ReliabilityComponent(nullptr)
+, ReliabilityComponent(NULL)
 {
 	OutgoingPacket.SetAllowResize(true);
 	OutgoingPacket.AllowAppend(true);
@@ -64,63 +64,58 @@ void PacketHandler::Initialize(Handler::Mode InMode)
 
 	Mode = InMode;
 
-	FString Components;
-	GConfig->GetString(TEXT("PacketHandlerComponents"), TEXT("Components"), Components, GEngineIni);
 
+	TArray<FString> Components;
 	TArray<ComponentAndOptions> ComponentsArray;
-	ComponentsArray.Add(ComponentAndOptions());
 
-	int32 i = 0;
-	int32 ComponentsCount = 0;
-	while(i < Components.Len())
-	{
-		TCHAR c = Components[i];
+	GConfig->GetArray(TEXT("PacketHandlerComponents"), TEXT("Components"), Components, GEngineIni);
 
-		// End of string
-		if(c == ',')
+	for (FString CurComponent : Components)
+	{
+		if (CurComponent.IsEmpty())
 		{
-			// If not end of string add a new string
-			if(i + 1 < Components.Len())
-			{
-				ComponentsArray.Add(ComponentAndOptions());
-				++ComponentsCount;
-			}
+			continue;
 		}
-		// Parsing Options
-		else if(c == '(')
+
+
+		int CurIndex = ComponentsArray.Add(ComponentAndOptions());
+
+		for (int32 i=0; i<CurComponent.Len(); i++)
 		{
-			// Skip '('
-			++i;
+			TCHAR c = CurComponent[i];
 
-			// Parse until end of options
-			while(i < Components.Len())
+			// Parsing Options
+			if (c == '(')
 			{
-				c = Components[i];
+				// Skip '('
+				++i;
 
-				// End of options
-				if(c == ')')
+				// Parse until end of options
+				for (; i<CurComponent.Len(); i++)
 				{
-					break;
+					c = CurComponent[i];
+
+					// End of options
+					if (c == ')')
+					{
+						break;
+					}
+					// Append char to options
+					else
+					{
+						ComponentsArray[CurIndex].Options.AppendChar(c);
+					}
 				}
-				// Append char to options
-				else
-				{
-					ComponentsArray[ComponentsCount].Options.AppendChar(c);
-				}	
-
-				++i;
+			}
+			// Append char to component name if not whitespace
+			else if (c != ' ')
+			{
+				ComponentsArray[CurIndex].ComponentName.AppendChar(c);
 			}
 		}
-		// Append char to component name if not whitespace
-		else if(c != ' ' )
-		{
-			ComponentsArray[ComponentsCount].ComponentName.AppendChar(c);
-		}
-
-		++i;
 	}
 
-	for (i = 0; i < ComponentsArray.Num(); i++)
+	for (int32 i=0; i<ComponentsArray.Num(); i++)
 	{
 		// Skip adding reliability component as it is added automatically in the handler
 		if(ComponentsArray[i].ComponentName == TEXT("ReliabilityHandlerComponent"))
@@ -128,24 +123,20 @@ void PacketHandler::Initialize(Handler::Mode InMode)
 			continue;
 		}
 
+
 		TSharedPtr<IModuleInterface> Interface = FModuleManager::Get().LoadModule(FName(*ComponentsArray[i].ComponentName));
 
 		if(Interface.IsValid())
 		{
-			TSharedPtr<FPacketHandlerComponentModuleInterface> PacketHandlerInterface(static_cast<FPacketHandlerComponentModuleInterface*>(&(*Interface)));
+			TSharedPtr<FPacketHandlerComponentModuleInterface> PacketHandlerInterface =
+				StaticCastSharedPtr<FPacketHandlerComponentModuleInterface>(Interface);
 
 			if(PacketHandlerInterface.IsValid())
 			{
-				// Options Specified
-				if(ComponentsArray[i].Options.Len() > 0)
-				{
-					Add(PacketHandlerInterface->CreateComponentInstance(ComponentsArray[i].Options));
-				}
-				// No Options Specified
-				else
-				{
-					Add(PacketHandlerInterface->CreateComponentInstance());
-				}
+				UE_LOG(PacketHandlerLog, Log, TEXT("Loading PacketHandler component: %s (%s)"), *ComponentsArray[i].ComponentName,
+						*ComponentsArray[i].Options);
+
+				Add(PacketHandlerInterface->CreateComponentInstance(ComponentsArray[i].Options));
 			}
 		}
 		else
@@ -155,11 +146,19 @@ void PacketHandler::Initialize(Handler::Mode InMode)
 	}
 }
 
-void PacketHandler::Add(HandlerComponent* NewHandler)
+void PacketHandler::Add(TSharedPtr<HandlerComponent> NewHandler)
 {
 	if (State != Handler::State::Uninitialized)
 	{
 		LowLevelFatalError(TEXT("Handler added during runtime."));
+		return;
+	}
+
+	// This should always be fatal, as an unexpectedly missing handler, may break net compatibility with the remote server/client
+	if (!NewHandler.IsValid())
+	{
+		LowLevelFatalError(TEXT("Failed to add handler - invalid instance."));
+		return;
 	}
 
 	HandlerComponents.Add(NewHandler);
@@ -175,9 +174,9 @@ const ProcessedPacket PacketHandler::Outgoing(uint8* Packet, int32 Count)
 	{
 		case Handler::State::Uninitialized:
 		{
-			if(ReliabilityComponent == nullptr)
+			if (!ReliabilityComponent.IsValid())
 			{
-				ReliabilityComponent = new ReliabilityHandlerComponent;
+				ReliabilityComponent = MakeShareable(new ReliabilityHandlerComponent);
 				Add(ReliabilityComponent);
 			}
 
@@ -224,7 +223,7 @@ const ProcessedPacket PacketHandler::Outgoing(uint8* Packet, int32 Count)
 	}
 
 	// Queue packet for resending before handling
-	if(Count > 0 && ReliabilityComponent != nullptr)
+	if(Count > 0 && ReliabilityComponent.IsValid())
 	{
 		ReliabilityComponent->QueuePacketForResending(Packet, Count);
 	}
@@ -249,9 +248,9 @@ const ProcessedPacket PacketHandler::Incoming(uint8* Packet, int32 Count)
 	{
 		case Handler::State::Uninitialized:
 		{
-			if (ReliabilityComponent == nullptr)
+			if (!ReliabilityComponent.IsValid())
 			{
-				ReliabilityComponent = new ReliabilityHandlerComponent;
+				ReliabilityComponent = MakeShareable(new ReliabilityHandlerComponent);
 				Add(ReliabilityComponent);
 			}
 
diff --git a/Engine/Source/Runtime/PacketHandlers/PacketHandler/Public/PacketHandler.h b/Engine/Source/Runtime/PacketHandlers/PacketHandler/Public/PacketHandler.h
--- a/Engine/Source/Runtime/PacketHandlers/PacketHandler/Public/PacketHandler.h
+++ b/Engine/Source/Runtime/PacketHandlers/PacketHandler/Public/PacketHandler.h
@@ -84,7 +84,7 @@ public:
 	virtual void Tick(float DeltaTime);
 
 	/* For adding handler components */
-	virtual void Add(HandlerComponent* NewHandler);
+	virtual void Add(TSharedPtr<HandlerComponent> NewHandler);
 
 	/* Handles any incoming packets */
 	virtual const ProcessedPacket Incoming(uint8* Packet, int32 Count);
@@ -121,7 +121,7 @@ private:
 	FBitReader IncomingPacket;
 
 	/* Array of Handler components */
-	TArray<HandlerComponent*> HandlerComponents;
+	TArray<TSharedPtr<HandlerComponent>> HandlerComponents;
 
 	/* State of the handler */
 	Handler::State State;
@@ -138,7 +138,7 @@ private:
 	TQueue<BufferedPacket*> QueuedPackets;
 
 	/* Reliability Handler Component */
-	ReliabilityHandlerComponent* ReliabilityComponent;
+	TSharedPtr<ReliabilityHandlerComponent> ReliabilityComponent;
 };
 
 /*
@@ -211,7 +211,6 @@ class FPacketHandlerComponentModuleInterface : public IModuleInterface
 {
 public:
 	/* Creates an instance of this component */
-	virtual HandlerComponent* CreateComponentInstance() PURE_VIRTUAL(FPacketHandlerModuleInterface::CreateComponentInstance,return nullptr;);
-	/* Creates an instance of this component with an FString of options */
-	virtual HandlerComponent* CreateComponentInstance(FString& Options) PURE_VIRTUAL(FPacketHandlerModuleInterface::CreateComponentInstance,return nullptr;);
+	virtual TSharedPtr<HandlerComponent> CreateComponentInstance(FString& Options)
+		PURE_VIRTUAL(FPacketHandlerModuleInterface::CreateComponentInstance, return TSharedPtr<HandlerComponent>(NULL););
 };
diff --git a/Engine/Source/Runtime/PacketHandlers/ReliabilityHandlerComponent/Private/ReliabilityHandlerComponent.cpp b/Engine/Source/Runtime/PacketHandlers/ReliabilityHandlerComponent/Private/ReliabilityHandlerComponent.cpp
--- a/Engine/Source/Runtime/PacketHandlers/ReliabilityHandlerComponent/Private/ReliabilityHandlerComponent.cpp
+++ b/Engine/Source/Runtime/PacketHandlers/ReliabilityHandlerComponent/Private/ReliabilityHandlerComponent.cpp
@@ -179,7 +179,7 @@ void ReliabilityHandlerComponent::QueuePacketForResending(uint8* Packet, int32 C
 }
 
 // MODULE INTERFACE
-HandlerComponent* FReliabilityHandlerComponentModuleInterface::CreateComponentInstance()
+TSharedPtr<HandlerComponent> FReliabilityHandlerComponentModuleInterface::CreateComponentInstance(FString& Options)
 {
-	return new ReliabilityHandlerComponent;
+	return MakeShareable(new ReliabilityHandlerComponent);
 }
\ No newline at end of file
diff --git a/Engine/Source/Runtime/PacketHandlers/ReliabilityHandlerComponent/Public/ReliabilityHandlerComponent.h b/Engine/Source/Runtime/PacketHandlers/ReliabilityHandlerComponent/Public/ReliabilityHandlerComponent.h
--- a/Engine/Source/Runtime/PacketHandlers/ReliabilityHandlerComponent/Public/ReliabilityHandlerComponent.h
+++ b/Engine/Source/Runtime/PacketHandlers/ReliabilityHandlerComponent/Public/ReliabilityHandlerComponent.h
@@ -55,5 +55,5 @@ protected:
 class FReliabilityHandlerComponentModuleInterface : public FPacketHandlerComponentModuleInterface
 {
 public:
-	virtual HandlerComponent* CreateComponentInstance() override;
-};
\ No newline at end of file
+	virtual TSharedPtr<HandlerComponent> CreateComponentInstance(FString& Options) override;
+};
