Compile fix - virtual destructor for PacketHandler's.

diff --git a/Engine/Source/Runtime/PacketHandlers/PacketHandler/Public/PacketHandler.h b/Engine/Source/Runtime/PacketHandlers/PacketHandler/Public/PacketHandler.h
--- a/Engine/Source/Runtime/PacketHandlers/PacketHandler/Public/PacketHandler.h
+++ b/Engine/Source/Runtime/PacketHandlers/PacketHandler/Public/PacketHandler.h
@@ -151,6 +151,9 @@ public:
 	/* Default initialization of data */
 	HandlerComponent();
 
+	/** Base destructor */
+	virtual ~HandlerComponent() {}
+
 	/* Returns whether this handler is currently active */
 	virtual bool IsActive() const;
 
