Updated DDoS detection, based on review.

diff --git a/Engine/Config/BaseEngine.ini b/Engine/Config/BaseEngine.ini
--- a/Engine/Config/BaseEngine.ini
+++ b/Engine/Config/BaseEngine.ini
@@ -788,6 +788,7 @@ MaxPortCountToTry=512
 [DDoSDetection]
 bDDoSDetection=false
 bDDoSAnalytics=false
+DDoSLogSpamLimit=64
 HitchTimeQuotaMS=500
 HitchFrameTolerance=3
 +DetectionSeverity=Burst
diff --git a/Engine/Plugins/Experimental/HTML5Networking/Source/HTML5Networking/Private/WebsocketConnection.cpp b/Engine/Plugins/Experimental/HTML5Networking/Source/HTML5Networking/Private/WebsocketConnection.cpp
--- a/Engine/Plugins/Experimental/HTML5Networking/Source/HTML5Networking/Private/WebsocketConnection.cpp
+++ b/Engine/Plugins/Experimental/HTML5Networking/Source/HTML5Networking/Private/WebsocketConnection.cpp
@@ -225,7 +225,7 @@ TSharedPtr<FInternetAddr> UWebSocketConnection::GetInternetAddr()
 {
 	struct sockaddr_in* sock = WebSocket->GetRemoteAddr();
 
-	// @todo: This should be based on NetConnection.RemoteAddr, when moved down from IPConnection
+	// @todo #JIRA UENET-883: This should be based on NetConnection.RemoteAddr, when moved down from IPConnection
 	return ISocketSubsystem::Get()->CreateInternetAddr((int32)ntohl(sock->sin_addr.s_addr), (int32)ntohs(sock->sin_port));
 }
 
diff --git a/Engine/Plugins/Online/OnlineSubsystemOculus/Source/Private/OculusNetConnection.cpp b/Engine/Plugins/Online/OnlineSubsystemOculus/Source/Private/OculusNetConnection.cpp
--- a/Engine/Plugins/Online/OnlineSubsystemOculus/Source/Private/OculusNetConnection.cpp
+++ b/Engine/Plugins/Online/OnlineSubsystemOculus/Source/Private/OculusNetConnection.cpp
@@ -157,7 +157,7 @@ void UOculusNetConnection::FinishDestroy()
 
 TSharedPtr<FInternetAddr> UOculusNetConnection::GetInternetAddr()
 {
-	// @todo: This should be based on NetConnection.RemoteAddr, when moved down from IPConnection
+	// @todo #JIRA UENET-883: This should be based on NetConnection.RemoteAddr, when moved down from IPConnection
 	return MakeShareable(new FInternetAddrOculus(PeerID));
 }
 
diff --git a/Engine/Plugins/Online/OnlineSubsystemUtils/Source/OnlineSubsystemUtils/Classes/IpConnection.h b/Engine/Plugins/Online/OnlineSubsystemUtils/Source/OnlineSubsystemUtils/Classes/IpConnection.h
--- a/Engine/Plugins/Online/OnlineSubsystemUtils/Source/OnlineSubsystemUtils/Classes/IpConnection.h
+++ b/Engine/Plugins/Online/OnlineSubsystemUtils/Source/OnlineSubsystemUtils/Classes/IpConnection.h
@@ -19,7 +19,7 @@ class ONLINESUBSYSTEMUTILS_API UIpConnection : public UNetConnection
     GENERATED_UCLASS_BODY()
 	// Variables.
 
-	// @todo: This should be moved down to UNetConnection, now that GetInternetAddr is a thing.
+	// @todo #JIRA UENET-883: This should be moved down to UNetConnection, now that GetInternetAddr is a thing.
 	//			A lot of platforms reinvent the wheel in their own inefficient way here, despite having their own FInternetAddr type
 	TSharedPtr<FInternetAddr>	RemoteAddr;
 
diff --git a/Engine/Plugins/Online/OnlineSubsystemUtils/Source/OnlineSubsystemUtils/Private/IpNetDriver.cpp b/Engine/Plugins/Online/OnlineSubsystemUtils/Source/OnlineSubsystemUtils/Private/IpNetDriver.cpp
--- a/Engine/Plugins/Online/OnlineSubsystemUtils/Source/OnlineSubsystemUtils/Private/IpNetDriver.cpp
+++ b/Engine/Plugins/Online/OnlineSubsystemUtils/Source/OnlineSubsystemUtils/Private/IpNetDriver.cpp
@@ -233,8 +233,7 @@ void UIpNetDriver::TickDispatch(float DeltaTime)
 	// Process all incoming packets.
 	uint8 Data[MAX_PACKET_SIZE];
 	uint8* DataRef = Data;
-	FInternetAddrMapRef FromAddrMap(SocketSubsystem->CreateInternetAddr());
-	TSharedRef<FInternetAddr>& FromAddr = FromAddrMap.Element;
+	TSharedRef<FInternetAddr> FromAddr = SocketSubsystem->CreateInternetAddr();
 
 	for (; Socket != nullptr;)
 	{
@@ -278,7 +277,7 @@ void UIpNetDriver::TickDispatch(float DeltaTime)
 		{
 			ESocketErrors Error = SocketSubsystem->GetLastErrorCode();
 
-			if(Error == SE_EWOULDBLOCK || Error == SE_NO_ERROR)
+			if (Error == SE_EWOULDBLOCK || Error == SE_NO_ERROR)
 			{
 				// No data or no error?
 				break;
@@ -289,10 +288,10 @@ void UIpNetDriver::TickDispatch(float DeltaTime)
 				// enforced by the server
 				if (Error == SE_EMSGSIZE)
 				{
+					DDoS.IncBadPacketCounter();
+
 					if (MyServerConnection)
 					{
-						DDoS.IncBadPacketCounter();
-
 						if (*MyServerConnection->RemoteAddr == *FromAddr)
 						{
 							Connection = MyServerConnection;
@@ -341,7 +340,7 @@ void UIpNetDriver::TickDispatch(float DeltaTime)
 		}
 
 		// Very-early-out - the NetConnection per frame time limit, limits all packet processing
-		if (DDoS.bHitFrameNetConnLimit)
+		if (DDoS.ShouldBlockNetConnPackets())
 		{
 			if (bOk)
 			{
@@ -368,9 +367,9 @@ void UIpNetDriver::TickDispatch(float DeltaTime)
 
 		if (Connection == nullptr)
 		{
-			UNetConnection** Result = MappedClientConnections.Find(FromAddrMap);
+			UNetConnection** Result = MappedClientConnections.Find(FromAddr);
 
-			Connection = (Result != nullptr ? (UIpConnection*)*Result : nullptr);
+			Connection = (Result != nullptr ? Cast<UIpConnection>(*Result) : nullptr);
 
 			check(Connection == nullptr || *Connection->RemoteAddr == *FromAddr);
 		}
@@ -421,10 +420,10 @@ void UIpNetDriver::TickDispatch(float DeltaTime)
 			// If we didn't find a client connection, maybe create a new one.
 			if( !Connection )
 			{
-				if (DDoS.bDDoSDetection)
+				if (DDoS.IsDDoSDetectionEnabled())
 				{
 					// If packet limits were reached, stop processing
-					if (DDoS.bHitFrameNonConnLimit)
+					if (DDoS.ShouldBlockNonConnPackets())
 					{
 						DDoS.IncDroppedPacketCounter();
 						continue;
@@ -531,7 +530,7 @@ void UIpNetDriver::TickDispatch(float DeltaTime)
 			// Send the packet to the connection for processing.
 			if (Connection && !bIgnorePacket)
 			{
-				if (DDoS.bDDoSDetection)
+				if (DDoS.IsDDoSDetectionEnabled())
 				{
 					DDoS.IncNetConnPacketCounter();
 					DDoS.CondCheckNetConnLimits();
diff --git a/Engine/Source/Runtime/Engine/Classes/Engine/NetDriver.h b/Engine/Source/Runtime/Engine/Classes/Engine/NetDriver.h
--- a/Engine/Source/Runtime/Engine/Classes/Engine/NetDriver.h
+++ b/Engine/Source/Runtime/Engine/Classes/Engine/NetDriver.h
@@ -14,7 +14,7 @@
 #include "PacketHandler.h"
 #include "Channel.h"
 #include "DDoSDetection.h"
-#include "Net/IPAddressRef.h"
+#include "IPAddress.h"
 
 #include "NetDriver.generated.h"
 
@@ -32,6 +32,8 @@ class UNetConnection;
 class UReplicationDriver;
 struct FNetworkObjectInfo;
 
+using FConnectionMap = TMap<TSharedRef<FInternetAddr>, UNetConnection*, FDefaultSetAllocator, FInternetAddrKeyMapFuncs<UNetConnection*>>;
+
 extern ENGINE_API TAutoConsoleVariable<int32> CVarNetAllowEncryption;
 extern ENGINE_API int32 GNumSaturatedConnections;
 extern ENGINE_API int32 GNumSharedSerializationHit;
@@ -286,8 +288,7 @@ public:
 	TArray<UNetConnection*> ClientConnections;
 
 	/** Map of IP's to NetConnection's - for fast lookup, particularly under DDoS - only valid IP's mapped (e.g. excludes DemoNetConnection) */
-	UPROPERTY()
-	TMap<FInternetAddrMapRef, UNetConnection*> MappedClientConnections;
+	FConnectionMap MappedClientConnections;
 
 
 	/** Serverside PacketHandler for managing connectionless packets */
diff --git a/Engine/Source/Runtime/Engine/Private/NetworkDriver.cpp b/Engine/Source/Runtime/Engine/Private/NetworkDriver.cpp
--- a/Engine/Source/Runtime/Engine/Private/NetworkDriver.cpp
+++ b/Engine/Source/Runtime/Engine/Private/NetworkDriver.cpp
@@ -1260,7 +1260,7 @@ bool UNetDriver::InitBase(bool bInitAsClient, FNetworkNotify* InNotify, const FU
 
 		DDoS.Init(FMath::Clamp(NetServerMaxTickRate, 1, 1000));
 
-		if (DDoS.bDDoSDetection && DDoS.bDDoSAnalytics)
+		if (DDoS.IsDDoSDetectionEnabled() && DDoS.IsDDoSAnalyticsEnabled())
 		{
 			DDoS.NotifySeverityEscalation.BindLambda(
 				[this](FString SeverityCategory)
@@ -1287,10 +1287,7 @@ void UNetDriver::InitConnectionlessHandler()
 
 		if (ConnectionlessHandler.IsValid())
 		{
-			ConnectionlessHandler->bConnectionlessHandler = true;
-			ConnectionlessHandler->DDoS = &DDoS;
-
-			ConnectionlessHandler->Initialize(Handler::Mode::Server, MAX_PACKET_SIZE, true, AnalyticsProvider);
+			ConnectionlessHandler->Initialize(Handler::Mode::Server, MAX_PACKET_SIZE, true, AnalyticsProvider, &DDoS);
 
 			// Add handling for the stateless connect handshake, for connectionless packets, as the outermost layer
 			TSharedPtr<HandlerComponent> NewComponent =
@@ -2665,6 +2662,11 @@ void UNetDriver::AddReferencedObjects(UObject* InThis, FReferenceCollector& Coll
 		Collector.AddReferencedObject(Replicator->ObjectPtr, This);
 		Collector.AddReferencedObject(Replicator->ObjectClass, This);
 	}
+
+	for (FConnectionMap::TIterator It(This->MappedClientConnections); It; ++It)
+	{
+		Collector.AddReferencedObject(It.Value(), This);
+	}
 }
 
 #if DO_ENABLE_NET_TEST
@@ -4144,7 +4146,7 @@ void UNetDriver::AddClientConnection(UNetConnection * NewConnection)
 
 	if (ConnAddr.IsValid())
 	{
-		MappedClientConnections.Add(FInternetAddrMapRef(ConnAddr), NewConnection);
+		MappedClientConnections.Add(ConnAddr.ToSharedRef(), NewConnection);
 	}
 
 	if (ReplicationDriver)
@@ -4250,7 +4252,7 @@ void UNetDriver::RemoveClientConnection(UNetConnection* ClientConnectionToRemove
 
 	if (AddrToRemove.IsValid())
 	{
-		verify(MappedClientConnections.Remove(AddrToRemove) == 1);
+		verify(MappedClientConnections.Remove(AddrToRemove.ToSharedRef()) == 1);
 	}
 
 	if (ReplicationDriver)
diff --git a/Engine/Source/Runtime/Engine/Private/PacketHandlers/StatelessConnectHandlerComponent.cpp b/Engine/Source/Runtime/Engine/Private/PacketHandlers/StatelessConnectHandlerComponent.cpp
--- a/Engine/Source/Runtime/Engine/Private/PacketHandlers/StatelessConnectHandlerComponent.cpp
+++ b/Engine/Source/Runtime/Engine/Private/PacketHandlers/StatelessConnectHandlerComponent.cpp
@@ -242,7 +242,9 @@ void StatelessConnectHandlerComponent::SendConnectChallenge(FString ClientAddres
 		ChallengePacket.Serialize(Cookie, ARRAY_COUNT(Cookie));
 
 #if !UE_BUILD_SHIPPING
-		UE_CLOG((Handler->DDoS == nullptr || !Handler->DDoS->CheckLogRestrictions()), LogHandshake, Log,
+		FDDoSDetection* DDoS = Handler->GetDDoS();
+
+		UE_CLOG((DDoS == nullptr || !DDoS->CheckLogRestrictions()), LogHandshake, Log,
 				TEXT("SendConnectChallenge. Timestamp: %f, Cookie: %s" ), Timestamp, *FString::FromBlob(Cookie, ARRAY_COUNT(Cookie)));
 #endif
 
@@ -620,13 +622,15 @@ void StatelessConnectHandlerComponent::IncomingConnectionless(FString Address, F
 		{
 			Packet.SetError();
 
-			if (Handler->DDoS != nullptr)
+			FDDoSDetection* DDoS = Handler->GetDDoS();
+
+			if (DDoS != nullptr)
 			{
-				Handler->DDoS->IncBadPacketCounter();
+				DDoS->IncBadPacketCounter();
 			}
 
 #if !UE_BUILD_SHIPPING
-			UE_CLOG(Handler->DDoS == nullptr || !Handler->DDoS->CheckLogRestrictions(), LogHandshake, Log,
+			UE_CLOG(DDoS == nullptr || !DDoS->CheckLogRestrictions(), LogHandshake, Log,
 					TEXT("IncomingConnectionless: Error reading handshake packet."));
 #endif
 		}
diff --git a/Engine/Source/Runtime/Engine/Public/Net/IPAddressRef.h b/Engine/Source/Runtime/Engine/Public/Net/IPAddressRef.h
deleted file mode 100644
diff --git a/Engine/Source/Runtime/PacketHandlers/PacketHandler/Private/DDoSDetection.cpp b/Engine/Source/Runtime/PacketHandlers/PacketHandler/Private/DDoSDetection.cpp
--- a/Engine/Source/Runtime/PacketHandlers/PacketHandler/Private/DDoSDetection.cpp
+++ b/Engine/Source/Runtime/PacketHandlers/PacketHandler/Private/DDoSDetection.cpp
@@ -62,21 +62,56 @@
  * FDDoSDetection
  */
 
+FDDoSDetection::FDDoSDetection()
+	: bDDoSDetection(false)
+	, bDDoSAnalytics(false)
+	, bHitFrameNonConnLimit(false)
+	, bHitFrameNetConnLimit(false)
+	, DetectionSeverity()
+	, ActiveState(0)
+	, WorstActiveState(0)
+	, LastMetEscalationConditions(0.0)
+	, bMetEscalationConditionsThisFrame(false)
+	, bDDoSLogRestrictions(false)
+	, DDoSLogSpamLimit(0)
+	, LogHitCounter(0)
+	, HitchTimeQuotaMS(-1)
+	, HitchFrameTolerance(-1)
+	, HitchFrameCount(0)
+	, LastPerSecQuotaBegin(0.0)
+	, CounterPerSecHistory()
+	, LastCounterPerSecHistoryIdx(0)
+	, StartFrameRecvTimestamp(0.0)
+	, EndFrameRecvTimestamp(0.0)
+	, StartFramePacketCount(0)
+	, ExpectedFrameTime(0.0)
+	, FrameAdjustment(0.f)
+	, NotifySeverityEscalation()
+{
+}
+
 void FDDoSDetection::Init(int32 MaxTickRate)
 {
 	ExpectedFrameTime = 1.0 / (MaxTickRate > 0.0 ? MaxTickRate : 30.0);
 
+	InitConfig();
+}
+
+void FDDoSDetection::InitConfig()
+{
 	const TCHAR* DDoSSection = TEXT("DDoSDetection");
-	int32 HitchFrameTolerance32 = 0;
+	int32 HitchFrameTolerance32 = -1;
 
 	GConfig->GetBool(DDoSSection, TEXT("bDDoSDetection"), bDDoSDetection, GEngineIni);
 	GConfig->GetBool(DDoSSection, TEXT("bDDoSAnalytics"), bDDoSAnalytics, GEngineIni);
+	GConfig->GetInt(DDoSSection, TEXT("DDoSLogSpamLimit"), DDoSLogSpamLimit, GEngineIni);
 	GConfig->GetInt(DDoSSection, TEXT("HitchTimeQuotaMS"), HitchTimeQuotaMS, GEngineIni);
+	GConfig->GetInt(DDoSSection, TEXT("HitchFrameTolerance"), HitchFrameTolerance32, GEngineIni);
 
-	if (GConfig->GetInt(DDoSSection, TEXT("HitchFrameTolerance"), HitchFrameTolerance32, GEngineIni))
-	{
-		HitchFrameTolerance = HitchFrameTolerance32;
-	}
+	HitchFrameTolerance = HitchFrameTolerance32;
+	DDoSLogSpamLimit = DDoSLogSpamLimit > 0 ? DDoSLogSpamLimit : 64;
+
+	DetectionSeverity.Empty();
 
 	if (bDDoSDetection)
 	{
@@ -85,7 +120,7 @@ void FDDoSDetection::Init(int32 MaxTickRate)
 
 		GConfig->GetArray(DDoSSection, TEXT("DetectionSeverity"), SeverityCatagories, GEngineIni);
 
-		for (FString CurCategory : SeverityCatagories)
+		for (const FString& CurCategory : SeverityCatagories)
 		{
 			FString CurSection = FString(DDoSSection) + TEXT(".") + CurCategory;
 
@@ -251,8 +286,11 @@ void FDDoSDetection::PreFrameReceive(float DeltaTime)
 
 		if (((StartFrameRecvTimestamp - LastPerSecQuotaBegin) - 1.0) > 0.0)
 		{
-			UE_CLOG(DroppedPacketCounter > 0, PacketHandlerLog, Warning, TEXT("DDoS Detection dropped '%i' packets during last second."),
-					DroppedPacketCounter);
+			UE_CLOG(DroppedPacketCounter > 0, PacketHandlerLog, Warning,
+				TEXT("DDoS Detection dropped '%i' packets during last second (bHitFrameNonConnLimit: %i, bHitFrameNetConnLimit: %i, ")
+				TEXT("DetectionSeverity: %s)."),
+				DroppedPacketCounter, (int32)bHitFrameNonConnLimit, (int32)bHitFrameNetConnLimit,
+				*DetectionSeverity[ActiveState].SeverityCategory);
 
 
 			// Record the last quota
@@ -274,17 +312,33 @@ void FDDoSDetection::PreFrameReceive(float DeltaTime)
 		}
 
 		StartFramePacketCount = NonConnPacketCounter;
+
+		if (LogHitCounter >= DDoSLogSpamLimit)
+		{
+			UE_LOG(PacketHandlerLog, Warning, TEXT("Previous frame hit DDoS LogHitCounter limit - hit count: %i (Max: %i)"), LogHitCounter,
+					DDoSLogSpamLimit);
+		}
+
+		LogHitCounter = 0;
+		bHitFrameNonConnLimit = false;
+		bHitFrameNetConnLimit = false;
 	}
+}
 
-	if (LogHitCounter >= DDOS_LOGSPAM_LIMIT)
+void FDDoSDetection::PostFrameReceive()
+{
+	if (bDDoSDetection)
 	{
-		UE_LOG(PacketHandlerLog, Warning, TEXT("Previous frame hit DDoS LogHitCounter limit - hit count: %i (Max: %i)"), LogHitCounter,
-				DDOS_LOGSPAM_LIMIT);
-	}
+		// Some packet counters require an end-frame check for DDoS detection
+		CheckNonConnQuotasAndLimits();
+
+
+		EndFrameRecvTimestamp = FPlatformTime::Seconds();
 
-	LogHitCounter = 0;
-	bHitFrameNonConnLimit = false;
-	bHitFrameNetConnLimit = false;
+		int32 FrameReceiveTimeMS = (int32)((EndFrameRecvTimestamp - StartFrameRecvTimestamp) * 1000.0);
+
+		WorstFrameReceiveTimeMS = FMath::Max(FrameReceiveTimeMS, WorstFrameReceiveTimeMS);
+	}
 }
 
 bool FDDoSDetection::CheckNonConnQuotasAndLimits()
diff --git a/Engine/Source/Runtime/PacketHandlers/PacketHandler/Private/PacketHandler.cpp b/Engine/Source/Runtime/PacketHandlers/PacketHandler/Private/PacketHandler.cpp
--- a/Engine/Source/Runtime/PacketHandlers/PacketHandler/Private/PacketHandler.cpp
+++ b/Engine/Source/Runtime/PacketHandlers/PacketHandler/Private/PacketHandler.cpp
@@ -79,15 +79,19 @@ void PacketHandler::Tick(float DeltaTime)
 	}
 }
 
-void PacketHandler::Initialize(Handler::Mode InMode, uint32 InMaxPacketBits, bool bConnectionlessOnly/*=false*/, TSharedPtr<IAnalyticsProvider> InProvider/*=nullptr*/)
+void PacketHandler::Initialize(Handler::Mode InMode, uint32 InMaxPacketBits, bool bConnectionlessOnly/*=false*/,
+								TSharedPtr<IAnalyticsProvider> InProvider/*=nullptr*/, FDDoSDetection* InDDoS/*=nullptr*/)
 {
 	Mode = InMode;
 	MaxPacketBits = InMaxPacketBits;
+	DDoS = InDDoS;
 
 	// @todo #JohnB: Redo this, so you don't load from the .ini at all, have it hardcoded elsewhere - do not want this in shipping.
 
+	bConnectionlessHandler = bConnectionlessOnly;
+
 	// Only UNetConnection's will load the .ini components, for now.
-	if (!bConnectionlessOnly)
+	if (!bConnectionlessHandler)
 	{
 		TArray<FString> Components;
 
diff --git a/Engine/Source/Runtime/PacketHandlers/PacketHandler/Public/DDoSDetection.h b/Engine/Source/Runtime/PacketHandlers/PacketHandler/Public/DDoSDetection.h
--- a/Engine/Source/Runtime/PacketHandlers/PacketHandler/Public/DDoSDetection.h
+++ b/Engine/Source/Runtime/PacketHandlers/PacketHandler/Public/DDoSDetection.h
@@ -7,12 +7,6 @@
 #include "HAL/PlatformTime.h"
 
 
-// Defines
-
-/** The limit for potential DDoS logspam, before log restrictions are enabled */
-#define DDOS_LOGSPAM_LIMIT 64
-
-
 // Delegates
 
 /**
@@ -103,7 +97,7 @@ struct PACKETHANDLER_API FDDoSState
 	 * @param TimePassedMS	The amount of time passed in milliseconds, since the beginning of this state
 	 * @return				Whether or not the quota was hit
 	 */
-	FORCEINLINE bool HasHitQuota(FDDoSPacketCounters& InCounter, int32 TimePassedMS) const
+	bool HasHitQuota(FDDoSPacketCounters& InCounter, int32 TimePassedMS) const
 	{
 		const bool bAtQuota = EscalateQuotaPacketsPerSec > 0 && InCounter.NonConnPacketCounter >= EscalateQuotaPacketsPerSec;
 		const bool bAtBadQuota = EscalateQuotaBadPacketsPerSec > 0 && InCounter.BadPacketCounter >= EscalateQuotaBadPacketsPerSec;
@@ -127,7 +121,7 @@ struct PACKETHANDLER_API FDDoSStateConfig : public FDDoSState
 	{
 	}
 
-	FORCEINLINE void ApplyState(FDDoSState& Target)
+	void ApplyState(FDDoSState& Target)
 	{
 		Target.EscalateQuotaPacketsPerSec		= EscalateQuotaPacketsPerSec;
 		Target.EscalateQuotaBadPacketsPerSec	= EscalateQuotaBadPacketsPerSec;
@@ -142,7 +136,7 @@ struct PACKETHANDLER_API FDDoSStateConfig : public FDDoSState
 	 * Applies only the per-frame adjusted state (based on expected vs actual framerate), to active DDoS protection.
 	 * ApplyState should be called, first.
 	 */
-	FORCEINLINE void ApplyAdjustedState(FDDoSState& Target, float FrameAdjustment)
+	void ApplyAdjustedState(FDDoSState& Target, float FrameAdjustment)
 	{
 		// Exclude escalation triggers from this
 		//Target.EscalateTimeQuotaMSPerFrame		= EscalateTimeQuotaMSPerFrame * FrameAdjustment;
@@ -160,32 +154,11 @@ struct PACKETHANDLER_API FDDoSStateConfig : public FDDoSState
 class PACKETHANDLER_API FDDoSDetection : protected FDDoSPacketCounters, protected FDDoSState
 {
 public:
-	FDDoSDetection()
-		: bDDoSDetection(false)
-		, bDDoSAnalytics(false)
-		, bHitFrameNonConnLimit(false)
-		, bHitFrameNetConnLimit(false)
-		, NotifySeverityEscalation()
-		, DetectionSeverity()
-		, ActiveState(0)
-		, WorstActiveState(0)
-		, LastMetEscalationConditions(0.0)
-		, bMetEscalationConditionsThisFrame(false)
-		, bDDoSLogRestrictions(false)
-		, LogHitCounter(0)
-		, HitchTimeQuotaMS(-1)
-		, HitchFrameTolerance(-1)
-		, HitchFrameCount(0)
-		, LastPerSecQuotaBegin(0.0)
-		, CounterPerSecHistory()
-		, LastCounterPerSecHistoryIdx(0)
-		, StartFrameRecvTimestamp(0.0)
-		, EndFrameRecvTimestamp(0.0)
-		, StartFramePacketCount(0)
-		, ExpectedFrameTime(0.0)
-		, FrameAdjustment(0.f)
-	{
-	}
+	/**
+	 * Default constructor
+	 */
+	FDDoSDetection();
+
 
 	/**
 	 * Initializes the DDoS detection settings
@@ -194,6 +167,11 @@ public:
 	 */
 	void Init(int32 MaxTickRate);
 
+	/**
+	 * Initializes the settings from the .ini file - must support reloading of settings on-the-fly
+	 */
+	void InitConfig();
+
 	/**
 	 * Updates the current DDoS detection severity state
 	 *
@@ -210,20 +188,13 @@ public:
 	/**
 	 * Triggered after packet receive ends, during the current frame
 	 */
-	FORCEINLINE void PostFrameReceive()
-	{
-		EndFrameRecvTimestamp = FPlatformTime::Seconds();
-
-		int32 FrameReceiveTimeMS = (int32)((EndFrameRecvTimestamp - StartFrameRecvTimestamp) * 1000.0);
-
-		WorstFrameReceiveTimeMS = FMath::Max(FrameReceiveTimeMS, WorstFrameReceiveTimeMS);
-	}
+	void PostFrameReceive();
 
 
 	/**
 	 * Rate limited call to CheckNonConnQuotasAndLimits
 	 */
-	FORCEINLINE void CondCheckNonConnQuotasAndLimits()
+	void CondCheckNonConnQuotasAndLimits()
 	{
 		// Limit checks to once every 128 packets
 		if ((NonConnPacketCounter & 0x7F) == 0)
@@ -235,7 +206,7 @@ public:
 	/**
 	 * Rate limited call to CheckNetConnLimits
 	 */
-	FORCEINLINE void CondCheckNetConnLimits()
+	void CondCheckNetConnLimits()
 	{
 		// Limit checks to once every 128 packets
 		if ((NetConnPacketCounter & 0x7F) == 0)
@@ -248,24 +219,29 @@ public:
 	/**
 	 * Accessor for bDDoSLogRestrictions - doubles as a per-frame logspam counter, automatically disabling logs after a quota
 	 */
-	FORCEINLINE bool CheckLogRestrictions()
+	bool CheckLogRestrictions()
 	{
-		return bDDoSLogRestrictions || (bDDoSDetection && ++LogHitCounter > DDOS_LOGSPAM_LIMIT);
+		return bDDoSLogRestrictions || (bDDoSDetection && ++LogHitCounter > DDoSLogSpamLimit);
 	}
 
 
 	// Brief accessors
 
-	FORCEINLINE void IncNonConnPacketCounter()			{ ++NonConnPacketCounter; }
-	FORCEINLINE int32 GetNonConnPacketCounter() const	{ return NonConnPacketCounter; }
-	FORCEINLINE void IncNetConnPacketCounter()			{ ++NetConnPacketCounter; }
-	FORCEINLINE int32 GetNetConnPacketCounter() const	{ return NetConnPacketCounter; }
-	FORCEINLINE void IncBadPacketCounter()				{ ++BadPacketCounter; }
-	FORCEINLINE int32 GetBadPacketCounter() const		{ return BadPacketCounter; }
-	FORCEINLINE void IncErrorPacketCounter()			{ ++ErrorPacketCounter; }
-	FORCEINLINE int32 GetErrorPacketCounter() const		{ return ErrorPacketCounter; }
-	FORCEINLINE void IncDroppedPacketCounter()			{ ++DroppedPacketCounter; }
-	FORCEINLINE int32 GetDroppedPacketCounter() const	{ return DroppedPacketCounter; }
+	bool IsDDoSDetectionEnabled() const		{ return bDDoSDetection; }
+	bool IsDDoSAnalyticsEnabled() const		{ return bDDoSAnalytics; }
+	bool ShouldBlockNonConnPackets() const	{ return bHitFrameNonConnLimit; }
+	bool ShouldBlockNetConnPackets() const	{ return bHitFrameNetConnLimit; }
+
+	void IncNonConnPacketCounter()			{ ++NonConnPacketCounter; }
+	int32 GetNonConnPacketCounter() const	{ return NonConnPacketCounter; }
+	void IncNetConnPacketCounter()			{ ++NetConnPacketCounter; }
+	int32 GetNetConnPacketCounter() const	{ return NetConnPacketCounter; }
+	void IncBadPacketCounter()				{ ++BadPacketCounter; }
+	int32 GetBadPacketCounter() const		{ return BadPacketCounter; }
+	void IncErrorPacketCounter()			{ ++ErrorPacketCounter; }
+	int32 GetErrorPacketCounter() const		{ return ErrorPacketCounter; }
+	void IncDroppedPacketCounter()			{ ++DroppedPacketCounter; }
+	int32 GetDroppedPacketCounter() const	{ return DroppedPacketCounter; }
 
 protected:
 	/**
@@ -280,7 +256,7 @@ protected:
 	 *
 	 * @return	Whether or not NetColnnection packet limits have been reached
 	 */
-	FORCEINLINE bool CheckNetConnLimits()
+	bool CheckNetConnLimits()
 	{
 		return NetConnPacketTimeLimitMSPerFrame > 0 &&
 				(int32)((FPlatformTime::Seconds() - StartFrameRecvTimestamp) * 1000.0) > NetConnPacketTimeLimitMSPerFrame;
@@ -288,7 +264,7 @@ protected:
 
 
 
-public:
+protected:
 	/** Whether or not DDoS detection is presently enabled */
 	bool bDDoSDetection;
 
@@ -301,11 +277,7 @@ public:
 	/** Whether or not the current frame has reached NetConnection packet limits, and should block ALL further packets */
 	bool bHitFrameNetConnLimit;
 
-	/** Analytics delegate for notifying of severity state escalations */
-	FDDoSSeverityEscalation NotifySeverityEscalation;
-
 
-protected:
 	/** The different DDoS detection states, of escalating severity, depending on the strength of the DDoS */
 	TArray<FDDoSStateConfig> DetectionSeverity;
 
@@ -325,6 +297,9 @@ protected:
 	/** Whether or not restriction of log messages from non-NetConnection packets is enabled */
 	bool bDDoSLogRestrictions;
 
+	/** The maximum number of non-NetConnection triggered log messages per frame, before further logs are dropped this frame */
+	int32 DDoSLogSpamLimit;
+
 	/** Counter for log restriction hits, in the current frame */
 	int32 LogHitCounter;
 
@@ -363,4 +338,9 @@ protected:
 
 	/** The current frames adjustment/deviation, from ExpectedFrameTime */
 	float FrameAdjustment;
+
+
+public:
+	/** Analytics delegate for notifying of severity state escalations */
+	FDDoSSeverityEscalation NotifySeverityEscalation;
 };
diff --git a/Engine/Source/Runtime/PacketHandlers/PacketHandler/Public/PacketHandler.h b/Engine/Source/Runtime/PacketHandlers/PacketHandler/Public/PacketHandler.h
--- a/Engine/Source/Runtime/PacketHandlers/PacketHandler/Public/PacketHandler.h
+++ b/Engine/Source/Runtime/PacketHandlers/PacketHandler/Public/PacketHandler.h
@@ -198,8 +198,10 @@ public:
 	 * @param InMaxPacketBits		The maximum supported packet size
 	 * @param bConnectionlessOnly	Whether or not this is a connectionless-only manager (ignores .ini components)
 	 * @param InProvider			The analytics provider
+	 * @param InDDoS				Reference to the owning net drivers DDoS detection handler
 	 */
-	void Initialize(Handler::Mode Mode, uint32 InMaxPacketBits, bool bConnectionlessOnly=false, TSharedPtr<class IAnalyticsProvider> InProvider=nullptr);
+	void Initialize(Handler::Mode Mode, uint32 InMaxPacketBits, bool bConnectionlessOnly=false,
+					TSharedPtr<class IAnalyticsProvider> InProvider=nullptr, FDDoSDetection* InDDoS=nullptr);
 
 	/**
 	 * Used for external initialization of delegates
@@ -454,6 +456,9 @@ public:
 		return State == Handler::State::Initialized;
 	}
 
+	/** Returns a pointer to the DDoS detection handler */
+	FDDoSDetection* GetDDoS() const { return DDoS; }
+
 
 private:
 	/**
@@ -495,13 +500,14 @@ public:
 	/** Time, updated by Tick */
 	float Time;
 
+
+private:
 	/** Whether or not this PacketHandler handles connectionless (i.e. non-UNetConnection) data */
 	bool bConnectionlessHandler;
 
 	/** Mirroring UNetDriver.DDoS*/
 	FDDoSDetection* DDoS;
 
-private:
 	/** Delegate used for triggering PacketHandler/HandlerComponent-sourced sends */
 	FPacketHandlerLowLevelSend LowLevelSendDel;
 
diff --git a/Engine/Source/Runtime/Sockets/Public/IPAddress.h b/Engine/Source/Runtime/Sockets/Public/IPAddress.h
--- a/Engine/Source/Runtime/Sockets/Public/IPAddress.h
+++ b/Engine/Source/Runtime/Sockets/Public/IPAddress.h
@@ -110,14 +110,7 @@ public:
 	 *
 	 * @return	The value to use for the hash
 	 */
-	virtual uint32 GetTypeHash()
-	{
-		uint32 ThisIP;
-
-		GetIp(ThisIP);
-
-		return ThisIP + (GetPort() * 23);
-	}
+	virtual uint32 GetTypeHash() = 0;
 
 	/**
 	 * Is this a well formed internet address
@@ -349,3 +342,28 @@ public:
 		return *Addr;
 	}
 };
+
+/**
+ * KeyFuncs for mapping/hashing FInternetAddr shared references
+ *
+ * NOTE: Implements TSharedRef over TSharedPtr, as this is required for performance critical code
+ */
+template<typename ValueType>
+struct FInternetAddrKeyMapFuncs : public BaseKeyFuncs<ValueType, TSharedRef<FInternetAddr>, false>
+{
+	static FORCEINLINE const TSharedRef<FInternetAddr>& GetSetKey(const TPair<TSharedRef<FInternetAddr>, ValueType>& Element)
+	{
+		return Element.Key;
+	}
+
+	static FORCEINLINE bool Matches(const TSharedRef<FInternetAddr>& A, const TSharedRef<FInternetAddr>& B)
+	{
+		return *A == *B;
+	}
+
+	static FORCEINLINE uint32 GetKeyHash(const TSharedRef<FInternetAddr>& Key)
+	{
+		return Key->GetTypeHash();
+	}
+};
+
