Updating the check to make sure that uninitialized components are blocked from doing low level sends before they are entering initialization or are already initialized.

diff --git a/Engine/Source/Runtime/PacketHandlers/PacketHandler/Private/PacketHandler.cpp b/Engine/Source/Runtime/PacketHandlers/PacketHandler/Private/PacketHandler.cpp
--- a/Engine/Source/Runtime/PacketHandlers/PacketHandler/Private/PacketHandler.cpp
+++ b/Engine/Source/Runtime/PacketHandlers/PacketHandler/Private/PacketHandler.cpp
@@ -590,8 +590,8 @@ void PacketHandler::SendHandlerPacket(HandlerComponent* InComponent, FBitWriter&
 {
 	// @todo #JohnB: There is duplication between this function and others, it would be nice to reduce this.
 
-	// For the moment, this should only be used by HandlerComponents which are in the process of initializing.
-	check(State == Handler::State::InitializingComponents);
+	// Prevent any cases where a send happens before the handler is ready.
+	check(State != Handler::State::Uninitialized);
 
 	if (LowLevelSendDel.IsBound())
 	{
