Brought forward a change needed for FPacketIterator checkin.

diff --git a/Engine/Source/Runtime/PacketHandlers/PacketHandler/Public/DDoSDetection.h b/Engine/Source/Runtime/PacketHandlers/PacketHandler/Public/DDoSDetection.h
--- a/Engine/Source/Runtime/PacketHandlers/PacketHandler/Public/DDoSDetection.h
+++ b/Engine/Source/Runtime/PacketHandlers/PacketHandler/Public/DDoSDetection.h
@@ -252,7 +252,7 @@ public:
 	int32 GetBadPacketCounter() const		{ return BadPacketCounter; }
 	void IncErrorPacketCounter()			{ ++ErrorPacketCounter; }
 	int32 GetErrorPacketCounter() const		{ return ErrorPacketCounter; }
-	void IncDroppedPacketCounter()			{ ++DroppedPacketCounter; }
+	void IncDroppedPacketCounter(int32 Count=1)			{ DroppedPacketCounter += Count; }
 	int32 GetDroppedPacketCounter() const	{ return DroppedPacketCounter; }
 
 protected:
