Merging using //UE4/Dev-Networking_to_//UE4/Release-4.21

Adding deprecation flags in to engine networking api functions that didn't get them. Also adding in temporary backwards compatibility.

diff --git a/Engine/Source/Runtime/Engine/Classes/Engine/DemoNetConnection.h b/Engine/Source/Runtime/Engine/Classes/Engine/DemoNetConnection.h
--- a/Engine/Source/Runtime/Engine/Classes/Engine/DemoNetConnection.h
+++ b/Engine/Source/Runtime/Engine/Classes/Engine/DemoNetConnection.h
@@ -26,6 +26,16 @@ struct FQueuedDemoPacket
 	uint32 SeenLevelIndex;
 
 public:
+	FORCEINLINE FQueuedDemoPacket(uint8* InData, int32 InSizeBytes, int32 InSizeBits) 
+		: Data()
+		, SizeBits(InSizeBits)
+		, Traits()
+		, SeenLevelIndex(0)
+	{
+		Data.AddUninitialized(InSizeBytes);
+		FMemory::Memcpy(Data.GetData(), InData, InSizeBytes);
+	}
+
 	FORCEINLINE FQueuedDemoPacket(uint8* InData, int32 InSizeBits, FOutPacketTraits& InTraits)
 		: Data()
 		, SizeBits(InSizeBits)
diff --git a/Engine/Source/Runtime/Engine/Classes/Engine/NetConnection.h b/Engine/Source/Runtime/Engine/Classes/Engine/NetConnection.h
--- a/Engine/Source/Runtime/Engine/Classes/Engine/NetConnection.h
+++ b/Engine/Source/Runtime/Engine/Classes/Engine/NetConnection.h
@@ -173,6 +173,17 @@ struct DelayedPacket
 	double SendTime;
 
 public:
+	DEPRECATED(4.21, "Use the constructor that takes PacketTraits for allowing for analytics and flags")
+	FORCEINLINE DelayedPacket(uint8* InData, int32 InSizeBytes, int32 InSizeBits)
+		: Data()
+		, SizeBits(InSizeBits)
+		, Traits()
+		, SendTime(0.0)
+	{
+		Data.AddUninitialized(InSizeBytes);
+		FMemory::Memcpy(Data.GetData(), InData, InSizeBytes);
+	}
+
 	FORCEINLINE DelayedPacket(uint8* InData, int32 InSizeBits, FOutPacketTraits& InTraits)
 		: Data()
 		, SizeBits(InSizeBits)
@@ -647,6 +658,13 @@ public:
 	/** Describe the connection. */
 	ENGINE_API virtual FString Describe();
 
+	DEPRECATED(4.21, "Use the method that allows for packet traits for analytics and modification")
+	ENGINE_API virtual void LowLevelSend(void* Data, int32 CountBytes, int32 CountBits)
+	{
+		FOutPacketTraits EmptyTraits;
+		LowLevelSend(Data, CountBits, EmptyTraits);
+	}
+
 	/**
 	 * Sends a byte stream to the remote endpoint using the underlying socket
 	 *
@@ -764,6 +782,11 @@ public:
 	 */
 	ENGINE_API virtual void InitConnection(UNetDriver* InDriver, EConnectionState InState, const FURL& InURL, int32 InConnectionSpeed=0, int32 InMaxPacket=0);
 
+	DEPRECATED(4.21, "Analytics providers are now handled in the NetDriver")
+	ENGINE_API virtual void InitHandler(TSharedPtr<IAnalyticsProvider> InProvider)
+	{
+		InitHandler();
+	}
 
 	/**
 	 * Initializes the PacketHandler
diff --git a/Engine/Source/Runtime/Engine/Classes/Engine/NetDriver.h b/Engine/Source/Runtime/Engine/Classes/Engine/NetDriver.h
--- a/Engine/Source/Runtime/Engine/Classes/Engine/NetDriver.h
+++ b/Engine/Source/Runtime/Engine/Classes/Engine/NetDriver.h
@@ -726,6 +726,13 @@ public:
 	/** PostTick actions */
 	ENGINE_API virtual void PostTickFlush();
 
+	DEPRECATED(4.21, "Please use the LowLevelSend that requires packet traits for analytics and packet modifiers.")
+	ENGINE_API virtual void LowLevelSend(FString Address, void* Data, int32 CountBits)
+	{
+		FOutPacketTraits EmptyTraits;
+		LowLevelSend(Address, Data, CountBits, EmptyTraits);
+	}
+
 	/**
 	 * Sends a 'connectionless' (not associated with a UNetConection) packet, to the specified address.
 	 * NOTE: Address is an abstract format defined by subclasses. Anything calling this, must use an address supplied by the net driver.
diff --git a/Engine/Source/Runtime/Engine/Private/NetConnection.cpp b/Engine/Source/Runtime/Engine/Private/NetConnection.cpp
--- a/Engine/Source/Runtime/Engine/Private/NetConnection.cpp
+++ b/Engine/Source/Runtime/Engine/Private/NetConnection.cpp
@@ -291,7 +291,7 @@ void UNetConnection::InitHandler()
 		{
 			Handler::Mode Mode = Driver->ServerConnection != nullptr ? Handler::Mode::Client : Handler::Mode::Server;
 
-			Handler->InitializeDelegates(FPacketHandlerLowLevelSend::CreateUObject(this, &UNetConnection::LowLevelSend));
+			Handler->InitializeDelegates(FPacketHandlerLowLevelSendTraits::CreateUObject(this, &UNetConnection::LowLevelSend));
 			Handler->NotifyAnalyticsProvider(Driver->AnalyticsProvider, Driver->AnalyticsAggregator);
 			Handler->Initialize(Mode, MaxPacket * 8, false);
 
diff --git a/Engine/Source/Runtime/PacketHandlers/PacketHandler/Private/PacketHandler.cpp b/Engine/Source/Runtime/PacketHandlers/PacketHandler/Private/PacketHandler.cpp
--- a/Engine/Source/Runtime/PacketHandlers/PacketHandler/Private/PacketHandler.cpp
+++ b/Engine/Source/Runtime/PacketHandlers/PacketHandler/Private/PacketHandler.cpp
@@ -33,6 +33,7 @@ PacketHandler::PacketHandler(FDDoSDetection* InDDoS/*=nullptr*/)
 	, bConnectionlessHandler(false)
 	, DDoS(InDDoS)
 	, LowLevelSendDel()
+	, LowLevelSendDel_Deprecated()
 	, HandshakeCompleteDel()
 	, OutgoingPacket()
 	, IncomingPacket()
@@ -605,7 +606,7 @@ void PacketHandler::SendHandlerPacket(HandlerComponent* InComponent, FBitWriter&
 	// Prevent any cases where a send happens before the handler is ready.
 	check(State != Handler::State::Uninitialized);
 
-	if (LowLevelSendDel.IsBound())
+	if (LowLevelSendDel.IsBound() || LowLevelSendDel_Deprecated.IsBound())
 	{
 		bool bEncounteredComponent = false;
 
@@ -661,7 +662,8 @@ void PacketHandler::SendHandlerPacket(HandlerComponent* InComponent, FBitWriter&
 
 			bRawSend = true;
 
-			LowLevelSendDel.Execute(Writer.GetData(), Writer.GetNumBits(), Traits);
+			LowLevelSendDel_Deprecated.ExecuteIfBound(Writer.GetData(), Writer.GetNumBytes(), Writer.GetNumBits());
+			LowLevelSendDel.ExecuteIfBound(Writer.GetData(), Writer.GetNumBits(), Traits);
 
 			bRawSend = bOldRawSend;
 		}
diff --git a/Engine/Source/Runtime/PacketHandlers/PacketHandler/Public/PacketHandler.h b/Engine/Source/Runtime/PacketHandlers/PacketHandler/Public/PacketHandler.h
--- a/Engine/Source/Runtime/PacketHandlers/PacketHandler/Public/PacketHandler.h
+++ b/Engine/Source/Runtime/PacketHandlers/PacketHandler/Public/PacketHandler.h
@@ -26,7 +26,9 @@ class FNetAnalyticsAggregator;
  */
 
 // Delegate for allowing access to LowLevelSend, without a dependency upon Engine
-DECLARE_DELEGATE_ThreeParams(FPacketHandlerLowLevelSend, void* /* Data */, int32 /* CountBits */, FOutPacketTraits& /* Traits */);
+DECLARE_DELEGATE_ThreeParams(FPacketHandlerLowLevelSendTraits, void* /* Data */, int32 /* CountBits */, FOutPacketTraits& /* Traits */);
+
+DECLARE_DELEGATE_ThreeParams(FPacketHandlerLowLevelSend, void* /* Data */, int32 /* CountBytes */, int32 /* CountBits */);
 
 /**
  * Callback for notifying higher-level code that handshaking has completed, and that packets are now ready to send without buffering
@@ -149,6 +151,21 @@ private:
 	}
 
 public:
+	DEPRECATED(4.21, "Please use the new constructor that adds support for analytics and better precision")
+	BufferedPacket(uint8* InCopyData, uint32 InCountBits, float InResendTime=0.f, uint32 InId=0)
+		: CountBits(InCountBits)
+		, Traits()
+		, ResendTime(double(InResendTime))
+		, Id(InId)
+		, Address()
+		, FromComponent(nullptr)
+	{
+		check(InCopyData != nullptr);
+
+		Data = new uint8[FMath::DivideAndRoundUp(InCountBits, 8u)];
+		FMemory::Memcpy(Data, InCopyData, FMath::DivideAndRoundUp(InCountBits, 8u));
+	}
+
 	BufferedPacket(uint8* InCopyData, uint32 InCountBits, FOutPacketTraits& InTraits, double InResendTime=0.0, uint32 InId=0)
 		: CountBits(InCountBits)
 		, Traits(InTraits)
@@ -163,6 +180,23 @@ public:
 		FMemory::Memcpy(Data, InCopyData, FMath::DivideAndRoundUp(InCountBits, 8u));
 	}
 
+	DEPRECATED(4.21, "Please use the new constructor that adds support for analytics")
+	BufferedPacket(const FString& InAddress, uint8* InCopyData, uint32 InCountBits, double InResendTime=0.0, uint32 InId=0)
+		: CountBits(InCountBits)
+		, Traits()
+		, ResendTime(InResendTime)
+		, Id(InId)
+		, Address()
+		, FromComponent(nullptr)
+	{
+		check(InCopyData != nullptr);
+
+		Data = new uint8[FMath::DivideAndRoundUp(InCountBits, 8u)];
+		FMemory::Memcpy(Data, InCopyData, FMath::DivideAndRoundUp(InCountBits, 8u));
+
+		Address = InAddress;
+	}
+
 	BufferedPacket(const FString& InAddress, uint8* InCopyData, uint32 InCountBits, FOutPacketTraits& InTraits, double InResendTime=0.0, uint32 InId=0)
 		: BufferedPacket(InCopyData, InCountBits, InTraits, InResendTime, InId)
 	{
@@ -203,12 +237,18 @@ public:
 	void Initialize(Handler::Mode Mode, uint32 InMaxPacketBits, bool bConnectionlessOnly=false,
 					TSharedPtr<class IAnalyticsProvider> InProvider=nullptr, FDDoSDetection* InDDoS=nullptr);
 
+	DEPRECATED(4.21, "Use the traits based delegate instead for compatibility with other systems.")
+	void InitializeDelegates(FPacketHandlerLowLevelSend InLowLevelSendDel)
+	{
+		LowLevelSendDel_Deprecated = InLowLevelSendDel;
+	}
+
 	/**
 	 * Used for external initialization of delegates
 	 *
 	 * @param InLowLevelSendDel		The delegate the PacketHandler should use for triggering packet sends
 	 */
-	void InitializeDelegates(FPacketHandlerLowLevelSend InLowLevelSendDel)
+	void InitializeDelegates(FPacketHandlerLowLevelSendTraits InLowLevelSendDel)
 	{
 		LowLevelSendDel = InLowLevelSendDel;
 	}
@@ -298,6 +338,13 @@ public:
 		return Incoming_Internal(Packet, CountBytes, false, EmptyString);
 	}
 
+	DEPRECATED(4.21, "Please move to the functional flow that includes support for PacketTraits.")
+	FORCEINLINE const ProcessedPacket Outgoing(uint8* Packet, int32 CountBits)
+	{
+		FOutPacketTraits EmptyTraits;
+		return Outgoing(Packet, CountBits, EmptyTraits);
+	}
+
 	/**
 	 * Processes outgoing packets at the PacketHandler level, after all UNetConnection processing.
 	 *
@@ -330,6 +377,13 @@ public:
 		return Incoming_Internal(Packet, CountBytes, true, Address);
 	}
 
+	DEPRECATED(4.21, "Please use the member that supports PacketTraits for alllowing additional flags on sends.")
+	FORCEINLINE const ProcessedPacket OutgoingConnectionless(const FString& Address, uint8* Packet, int32 CountBits)
+	{
+		FOutPacketTraits EmptyTraits;
+		return OutgoingConnectionless(Address, Packet, CountBits, EmptyTraits);
+	}
+
 	/**
 	 * Processes outgoing packets without a UNetConnection, in the same manner as 'Outgoing' above
 	 * NOTE: Connectionless packets are unreliable.
@@ -376,6 +430,14 @@ protected:
 	const ProcessedPacket Outgoing_Internal(uint8* Packet, int32 CountBits, FOutPacketTraits& Traits, bool bConnectionless, const FString& Address );
 
 public:
+
+	DEPRECATED(4.21, "Please use the packet traits when sending to handle modifications of packets and analytics.")
+	void SendHandlerPacket(HandlerComponent* InComponent, FBitWriter& Writer)
+	{
+		FOutPacketTraits EmptyTraits;
+		SendHandlerPacket(InComponent, Writer, EmptyTraits);
+	}
+
 	/**
 	 * Send a packet originating from a HandlerComponent - will process through the HandlerComponents chain,
 	 * starting after the triggering component.
@@ -526,7 +588,10 @@ private:
 	FDDoSDetection* DDoS;
 
 	/** Delegate used for triggering PacketHandler/HandlerComponent-sourced sends */
-	FPacketHandlerLowLevelSend LowLevelSendDel;
+	FPacketHandlerLowLevelSendTraits LowLevelSendDel;
+
+	/** Delegate used for triggering PacketHandler/HandlerComponent-sourced sends (DEPRECATED) */
+	FPacketHandlerLowLevelSend LowLevelSendDel_Deprecated;
 
 	/** Delegate used for notifying that handshaking has completed */
 	FPacketHandlerHandshakeComplete HandshakeCompleteDel;
@@ -646,6 +711,13 @@ public:
 	 */
 	virtual void Incoming(FBitReader& Packet) = 0;
 
+	DEPRECATED(4.21, "Use the other Outgoing function as it allows for packet modifiers and traits.")
+	virtual void Outgoing(FBitWriter& Packet)
+	{
+		FOutPacketTraits EmptyTraits;
+		Outgoing(Packet, EmptyTraits);
+	}
+
 	/**
 	 * Handles any outgoing packets
 	 *
@@ -662,6 +734,13 @@ public:
 	 */
 	virtual void IncomingConnectionless(const FString& Address, FBitReader& Packet) = 0;
 
+	DEPRECATED(4.21, "Use the method that allows traits on the packet.")
+	virtual void OutgoingConnectionless(const FString& Address, FBitWriter& Packet)
+	{
+		FOutPacketTraits EmptyTraits;
+		OutgoingConnectionless(Address, Packet, EmptyTraits);
+	}
+
 	/**
 	 * Handles any outgoing packets not associated with a UNetConnection
 	 *
@@ -721,6 +800,11 @@ public:
 	/** Returns the name of this component. */
 	FName GetName() const { return Name; }
 
+	DEPRECATED(4.21, "The Analytics Provider is now handled in the main PacketHandler class.")
+	virtual void SetAnalyticsProvider(TSharedPtr<class IAnalyticsProvider> Provider)
+	{
+	}
+
 	/**
 	 * Notification that the analytics provider has been updated
 	 * NOTE: Can also mean disabled, e.g. during hotfix
diff --git a/Engine/Source/Runtime/PacketHandlers/ReliabilityHandlerComponent/Public/ReliabilityHandlerComponent.h b/Engine/Source/Runtime/PacketHandlers/ReliabilityHandlerComponent/Public/ReliabilityHandlerComponent.h
--- a/Engine/Source/Runtime/PacketHandlers/ReliabilityHandlerComponent/Public/ReliabilityHandlerComponent.h
+++ b/Engine/Source/Runtime/PacketHandlers/ReliabilityHandlerComponent/Public/ReliabilityHandlerComponent.h
@@ -28,6 +28,13 @@ public:
 	/* Queues a packet for resending */
 	void QueuePacketForResending(uint8* Packet, int32 CountBits, FOutPacketTraits& Traits);
 
+	DEPRECATED(4.21, "Use the PacketTraits version for sending packets with additional flags and options")
+	FORCEINLINE void QueueHandlerPacketForResending(HandlerComponent* InComponent, uint8* Packet, int32 CountBits)
+	{
+		FOutPacketTraits EmptyTraits;
+		QueueHandlerPacketForResending(InComponent, Packet, CountBits, EmptyTraits);
+	}
+
 	/**
 	 * Queues a packet sent through SendHandlerPacket, for resending
 	 *
