Copying //UE4/Release-Staging-4.21 to Dev-Main (//UE4/Dev-Main) @ 4547701
#rnx

diff --git a/Engine/Build/Android/Legacy/src/com/epicgames/ue4/LocalNotificationReceiver.java b/Engine/Build/Android/Legacy/src/com/epicgames/ue4/LocalNotificationReceiver.java
new file mode 100644
--- /dev/null
+++ b/Engine/Build/Android/Legacy/src/com/epicgames/ue4/LocalNotificationReceiver.java
@@ -0,0 +1,68 @@
+// Copyright 1998-2018 Epic Games, Inc. All Rights Reserved.
+
+ package com.epicgames.ue4;
+ 
+ import android.app.Notification;
+ import android.app.NotificationManager;
+ import android.app.PendingIntent;
+ import android.content.BroadcastReceiver;
+ import android.content.Context;
+ import android.content.Intent;
+ import android.support.v4.app.NotificationCompat;
+ 
+ public class LocalNotificationReceiver extends BroadcastReceiver
+ {
+ 
+	 public void onReceive(Context context, Intent intent)
+	 {
+		int notificationID = intent.getIntExtra("local-notification-ID" , 0);
+		String title  = intent.getStringExtra("local-notification-title");
+		String details  = intent.getStringExtra("local-notification-body");
+		String action = intent.getStringExtra("local-notification-action");
+		String activationEvent = intent.getStringExtra("local-notification-activationEvent");
+
+		if(title == null || details == null || action == null || activationEvent == null)
+		{
+			// Do not schedule any local notification if any allocation failed
+			return;
+		}
+
+		// Open UE4 app if clicked
+		Intent notificationIntent = new Intent(context, GameActivity.class);
+
+		// launch if closed but resume if running
+		notificationIntent.setFlags(Intent.FLAG_ACTIVITY_CLEAR_TOP | Intent.FLAG_ACTIVITY_SINGLE_TOP);
+		notificationIntent.putExtra("localNotificationID" , notificationID);
+		notificationIntent.putExtra("localNotificationAppLaunched" , true);
+		notificationIntent.putExtra("localNotificationLaunchActivationEvent", activationEvent);
+
+		int notificationIconID = context.getResources().getIdentifier("ic_notification", "drawable", context.getPackageName());
+		if (notificationIconID == 0)
+		{
+			notificationIconID = context.getResources().getIdentifier("icon", "drawable", context.getPackageName());
+		}
+		PendingIntent pendingNotificationIntent = PendingIntent.getActivity(context, notificationID, notificationIntent, 0);
+
+		NotificationCompat.Builder builder = new NotificationCompat.Builder(context)
+			.setSmallIcon(notificationIconID)
+			.setContentIntent(pendingNotificationIntent)
+			.setWhen(System.currentTimeMillis())
+			.setTicker(details)		// note: will not show up on Lollipop up except for accessibility
+			.setContentTitle(title);
+		if (android.os.Build.VERSION.SDK_INT >= 21)
+		{
+			builder.setColor(0xff0e1e43);
+		}
+		Notification notification = builder.build();
+
+		// Stick with the defaults
+		notification.flags |= Notification.FLAG_AUTO_CANCEL;
+		notification.defaults |= Notification.DEFAULT_SOUND | Notification.DEFAULT_VIBRATE;
+
+		NotificationManager notificationManager = (NotificationManager) context.getSystemService(Context.NOTIFICATION_SERVICE);
+
+		// show the notification
+		notificationManager.notify(notificationID, notification); 
+	 }
+ }
+ 
\ No newline at end of file
diff --git a/Engine/Config/BaseEngine.ini b/Engine/Config/BaseEngine.ini
--- a/Engine/Config/BaseEngine.ini
+++ b/Engine/Config/BaseEngine.ini
@@ -1769,6 +1769,7 @@ bSupportsPortraitOrientation=False
 bSupportsUpsideDownOrientation=False
 bSupportsLandscapeLeftOrientation=True
 bSupportsLandscapeRightOrientation=True
+PreferredLandscapeOrientation=LandscapeLeft
 bSupportsMetal=True
 bSupportsOpenGLES2=False
 bCookPVRTCTextures=True
diff --git a/Engine/Plugins/FX/Niagara/Source/Niagara/Private/NiagaraDataInterfaceTexture.cpp b/Engine/Plugins/FX/Niagara/Source/Niagara/Private/NiagaraDataInterfaceTexture.cpp
--- a/Engine/Plugins/FX/Niagara/Source/Niagara/Private/NiagaraDataInterfaceTexture.cpp
+++ b/Engine/Plugins/FX/Niagara/Source/Niagara/Private/NiagaraDataInterfaceTexture.cpp
@@ -388,7 +388,7 @@ struct FNiagaraDataInterfaceParametersCS_Texture : public FNiagaraDataInterfaceP
 			SetShaderValue(RHICmdList, ComputeShaderRHI, Dimensions, TexDims);
 			return;
 		}
-		FTextureRHIParamRef TextureRHI = Texture->TextureReference.TextureReferenceRHI;
+		FTextureRHIParamRef TextureRHI = Texture->TextureReference.TextureReferenceRHI->GetReferencedTexture();
 		SetTextureParameter(
 			RHICmdList,
 			ComputeShaderRHI,
diff --git a/Engine/Plugins/Runtime/GoogleVR/GoogleVRHMD/Source/GoogleVRHMD/Private/GoogleVRHMD.cpp b/Engine/Plugins/Runtime/GoogleVR/GoogleVRHMD/Source/GoogleVRHMD/Private/GoogleVRHMD.cpp
--- a/Engine/Plugins/Runtime/GoogleVR/GoogleVRHMD/Source/GoogleVRHMD/Private/GoogleVRHMD.cpp
+++ b/Engine/Plugins/Runtime/GoogleVR/GoogleVRHMD/Source/GoogleVRHMD/Private/GoogleVRHMD.cpp
@@ -1353,14 +1353,9 @@ void FGoogleVRHMD::PostRenderViewFamily_RenderThread(FRHICommandListImmediate& R
 	}
 
 	uint64 result = 0;
-	bool isTextureReadyForReadback = false;
-	while (SentTextureCount < ReadbackTextureCount && RHICmdList.GetRenderQueryResult(ReadbackCopyQueries[SentTextureCount % kReadbackTextureCount], result, false)) {
-		isTextureReadyForReadback = true;
-		SentTextureCount++;
-	}
-
-	if (isTextureReadyForReadback) {
-		int latestReadbackTextureIndex = (SentTextureCount - 1) % kReadbackTextureCount;
+	if (RHICmdList.GetRenderQueryResult(ReadbackCopyQueries[SentTextureCount % kReadbackTextureCount], result, false))
+	{
+		int latestReadbackTextureIndex = SentTextureCount % kReadbackTextureCount;
 		GDynamicRHI->RHIReadSurfaceData(
 			ReadbackTextures[latestReadbackTextureIndex],
 			FIntRect(FIntPoint(0, 0),
@@ -1374,6 +1369,8 @@ void FGoogleVRHMD::PostRenderViewFamily_RenderThread(FRHICommandListImmediate& R
 			ReadbackTextureSizes[latestReadbackTextureIndex].X * 4,
 			instant_preview::PIXEL_FORMAT_BGRA,
 			ReadbackReferencePoses[latestReadbackTextureIndex]);
+
+		SentTextureCount++;
 	}
 }
 #endif  // GOOGLEVRHMD_SUPPORTED_INSTANT_PREVIEW_PLATFORMS
diff --git a/Engine/Plugins/Runtime/Steam/SteamVR/Source/SteamVR/Private/SteamVRStereoLayers.cpp b/Engine/Plugins/Runtime/Steam/SteamVR/Source/SteamVR/Private/SteamVRStereoLayers.cpp
--- a/Engine/Plugins/Runtime/Steam/SteamVR/Source/SteamVR/Private/SteamVRStereoLayers.cpp
+++ b/Engine/Plugins/Runtime/Steam/SteamVR/Source/SteamVR/Private/SteamVRStereoLayers.cpp
@@ -219,10 +219,7 @@ void FSteamVRHMD::UpdateStereoLayers_RenderThread()
 	}
 
 	// Metal is not supported yet
-	if (IsMetalPlatform(GMaxRHIShaderPlatform))
-	{
-		return;
-	}
+	check(!IsMetalPlatform(GMaxRHIShaderPlatform));
 
 	static const auto CVarMixLayerPriorities = IConsoleManager::Get().FindTConsoleVariableDataInt(TEXT("vr.StereoLayers.bMixLayerPriorities"));
 	const bool bUpdateLayerPriorities = (CVarMixLayerPriorities->GetValueOnRenderThread() == 0) && GetStereoLayersDirty();
@@ -378,8 +375,15 @@ void FSteamVRHMD::GetAllocatedTexture(uint32 LayerId, FTextureRHIRef &Texture, F
 }
 
 //=============================================================================
-IStereoLayers* FSteamVRHMD::GetStereoLayers ()
+IStereoLayers* FSteamVRHMD::GetStereoLayers()
 {
+	// Metal is not supported yet. Fall back to the default portable implementation
+	if (IsMetalPlatform(GMaxRHIShaderPlatform))
+	{
+		return FHeadMountedDisplayBase::GetStereoLayers();
+	}
+
+
 	if (VROverlay)
 	{
 		return this;
diff --git a/Engine/Shaders/Private/VolumetricFog.usf b/Engine/Shaders/Private/VolumetricFog.usf
--- a/Engine/Shaders/Private/VolumetricFog.usf
+++ b/Engine/Shaders/Private/VolumetricFog.usf
@@ -636,7 +636,7 @@ void LightScatteringCS(
 
 	if (all(GridCoordinate < VolumetricFog.GridSizeInt))
 	{
-		ScatteringAndExtinction = isnan(ScatteringAndExtinction) || isinf(ScatteringAndExtinction) ? 0 : ScatteringAndExtinction;
+		ScatteringAndExtinction = -min(-ScatteringAndExtinction, float4(0.f, 0.f, 0.f, 0.f));	// vector isnan/isinf are miscompiled by hlslcc (see UE-66179)
 		RWLightScattering[GridCoordinate] = ScatteringAndExtinction;
 	}
 }
@@ -681,4 +681,4 @@ void FinalIntegrationCS(
 
 		RWIntegratedLightScattering[LayerCoordinate] = float4(AccumulatedLighting, AccumulatedTransmittance);
 	}
-}
\ No newline at end of file
+}
diff --git a/Engine/Source/Developer/Apple/MetalShaderFormat/Private/MetalDerivedData.cpp b/Engine/Source/Developer/Apple/MetalShaderFormat/Private/MetalDerivedData.cpp
--- a/Engine/Source/Developer/Apple/MetalShaderFormat/Private/MetalDerivedData.cpp
+++ b/Engine/Source/Developer/Apple/MetalShaderFormat/Private/MetalDerivedData.cpp
@@ -215,11 +215,11 @@ bool FMetalShaderBytecodeCooker::Build(TArray<uint8>& OutData)
 		if (bUseSharedPCH)
         {
             CopyLocalFileToRemote(Job.InputPCHFile, RemoteInputPCHFile);
-			MetalParams = FString::Printf(TEXT("-include-pch %s %s %s %s %s -Wno-null-character -fbracket-depth=1024 %s %s %s %s -o %s"), *RemoteInputPCHFile, *Job.MinOSVersion, *Job.DebugInfo, *Job.MathMode, XcodeMajorVers >= 10 ? TEXT("-c") : TEXT(""), *Job.Standard, *Job.Defines, *IncludeArgs, *RemoteInputFile, *RemoteObjFile);
+			MetalParams = FString::Printf(TEXT("-include-pch %s %s %s %s %s -Wno-null-character -fbracket-depth=1024 %s %s %s %s -o %s"), *RemoteInputPCHFile, *Job.MinOSVersion, *Job.DebugInfo, *Job.MathMode, TEXT("-c"), *Job.Standard, *Job.Defines, *IncludeArgs, *RemoteInputFile, *RemoteObjFile);
         }
         else
         {
-            MetalParams = FString::Printf(TEXT("%s %s %s %s -Wno-null-character -fbracket-depth=1024 %s %s %s %s -o %s"), *Job.MinOSVersion, *Job.DebugInfo, *Job.MathMode, XcodeMajorVers >= 10 ? TEXT("-c") : TEXT(""), *Job.Standard, *Job.Defines, *IncludeArgs, *RemoteInputFile, *RemoteObjFile);
+            MetalParams = FString::Printf(TEXT("%s %s %s %s -Wno-null-character -fbracket-depth=1024 %s %s %s %s -o %s"), *Job.MinOSVersion, *Job.DebugInfo, *Job.MathMode, TEXT("-c"), *Job.Standard, *Job.Defines, *IncludeArgs, *RemoteInputFile, *RemoteObjFile);
         }
 	}
 
diff --git a/Engine/Source/Developer/Apple/MetalShaderFormat/Private/MetalShaderCompiler.cpp b/Engine/Source/Developer/Apple/MetalShaderFormat/Private/MetalShaderCompiler.cpp
--- a/Engine/Source/Developer/Apple/MetalShaderFormat/Private/MetalShaderCompiler.cpp
+++ b/Engine/Source/Developer/Apple/MetalShaderFormat/Private/MetalShaderCompiler.cpp
@@ -544,7 +544,8 @@ uint16 GetXcodeVersion(uint64& BuildVersion)
 	{
 		Version = 0; // No Xcode install is 0, so only text shaders will work
 		FString XcodePath = GetXcodePath();
-		if (XcodePath.Len() > 0)
+		// Because of where and when this is called you can't invoke it on Win->Mac builds
+		if (XcodePath.Len() > 0 && PLATFORM_MAC)
 		{
 			FString Path = FString::Printf(TEXT("%s/usr/bin/xcodebuild"), *XcodePath);
 			FString Result;
diff --git a/Engine/Source/Developer/AssetTools/Private/AssetRenameManager.cpp b/Engine/Source/Developer/AssetTools/Private/AssetRenameManager.cpp
--- a/Engine/Source/Developer/AssetTools/Private/AssetRenameManager.cpp
+++ b/Engine/Source/Developer/AssetTools/Private/AssetRenameManager.cpp
@@ -260,10 +260,10 @@ bool FAssetRenameManager::FixReferencesAndRename(const TArray<FAssetRenameData>&
 	TArray<FAssetRenameDataWithReferencers> AssetsToRename;
 	AssetsToRename.Reset(AssetsAndNames.Num());
 	// Avoid duplicates when adding MapBuildData to list
-	TMap<UObject*, bool> AssetsToRenameLookup;
+	TSet<UObject*> AssetsToRenameLookup;
 	for (const FAssetRenameData& AssetRenameData : AssetsAndNames)
 	{
-		AssetsToRenameLookup.FindOrAdd(AssetRenameData.Asset.Get());
+		AssetsToRenameLookup.Add(AssetRenameData.Asset.Get());
 	}
 	for (const FAssetRenameData& AssetRenameData : AssetsAndNames)
 	{
@@ -273,11 +273,15 @@ bool FAssetRenameManager::FixReferencesAndRename(const TArray<FAssetRenameData>&
 			UWorld* World = Cast<UWorld>(AssetRenameData.Asset.Get());
 			if (World && World->PersistentLevel && World->PersistentLevel->MapBuildData && !AssetsToRenameLookup.Contains(World->PersistentLevel->MapBuildData))
 			{
-				FString NewMapBuildDataName = AssetRenameData.NewName + TEXT("_BuiltData");
-				// Perform rename of MapBuildData before world otherwise original files left behind
-				AssetsToRename.EmplaceAt(0, FAssetRenameDataWithReferencers(FAssetRenameData(World->PersistentLevel->MapBuildData, AssetRenameData.NewPackagePath, NewMapBuildDataName)));
-				AssetsToRename[0].bOnlyFixSoftReferences = AssetRenameData.bOnlyFixSoftReferences;
-				AssetsToRenameLookup.Add(World->PersistentLevel->MapBuildData);
+				// Leave MapBuildData inside the map's package
+				if (World->PersistentLevel->MapBuildData->GetOutermost() != World->GetOutermost())
+				{
+					FString NewMapBuildDataName = AssetRenameData.NewName + TEXT("_BuiltData");
+					// Perform rename of MapBuildData before world otherwise original files left behind
+					AssetsToRename.EmplaceAt(0, FAssetRenameDataWithReferencers(FAssetRenameData(World->PersistentLevel->MapBuildData, AssetRenameData.NewPackagePath, NewMapBuildDataName)));
+					AssetsToRename[0].bOnlyFixSoftReferences = AssetRenameData.bOnlyFixSoftReferences;
+					AssetsToRenameLookup.Add(World->PersistentLevel->MapBuildData);
+				}
 			}
 		}
 
diff --git a/Engine/Source/Programs/UnrealBuildTool/Platform/Android/UEDeployAndroid.cs b/Engine/Source/Programs/UnrealBuildTool/Platform/Android/UEDeployAndroid.cs
--- a/Engine/Source/Programs/UnrealBuildTool/Platform/Android/UEDeployAndroid.cs
+++ b/Engine/Source/Programs/UnrealBuildTool/Platform/Android/UEDeployAndroid.cs
@@ -3331,6 +3331,10 @@ namespace UnrealBuildTool
 			CopyFileDirectory(UE4BuildFilesPath, UE4BuildPath, Replacements);
 			CopyFileDirectory(UE4BuildFilesPath + "/NotForLicensees", UE4BuildPath, Replacements);
 			CopyFileDirectory(UE4BuildFilesPath + "/NoRedist", UE4BuildPath, Replacements);
+			if (!bGradleEnabled)
+			{
+				CopyFileDirectory(Path.Combine(EngineDirectory, "Build", "Android", "Legacy"), UE4BuildPath, Replacements);
+			}
 			CopyFileDirectory(GameBuildFilesPath, UE4BuildPath, Replacements);
 			CopyFileDirectory(GameBuildFilesPath + "/NotForLicensees", UE4BuildPath, Replacements);
 			CopyFileDirectory(GameBuildFilesPath + "/NoRedist", UE4BuildPath, Replacements);
diff --git a/Engine/Source/Programs/UnrealBuildTool/ProjectFiles/Xcode/XcodeProject.cs b/Engine/Source/Programs/UnrealBuildTool/ProjectFiles/Xcode/XcodeProject.cs
--- a/Engine/Source/Programs/UnrealBuildTool/ProjectFiles/Xcode/XcodeProject.cs
+++ b/Engine/Source/Programs/UnrealBuildTool/ProjectFiles/Xcode/XcodeProject.cs
@@ -1050,7 +1050,7 @@ namespace UnrealBuildTool
 				{
 					foreach (UnrealTargetPlatform Platform in Platforms)
 					{
-						if (InstalledPlatformInfo.IsValidPlatform(Platform, EProjectType.Code) && (Platform == UnrealTargetPlatform.Mac || Platform == UnrealTargetPlatform.IOS)) // @todo support other platforms
+						if (InstalledPlatformInfo.IsValidPlatform(Platform, EProjectType.Code) && (Platform == UnrealTargetPlatform.Mac || Platform == UnrealTargetPlatform.IOS || Platform == UnrealTargetPlatform.TVOS)) // @todo support other platforms
 						{
 							UEBuildPlatform BuildPlatform = UEBuildPlatform.GetBuildPlatform(Platform, true);
 							if ((BuildPlatform != null) && (BuildPlatform.HasRequiredSDKsInstalled() == SDKStatus.Valid))
@@ -1114,7 +1114,7 @@ namespace UnrealBuildTool
                                                 string TVOSExecutableName = MacExecutableName.Replace("-Mac-", "-TVOS-");
                                                 BuildConfigs.Add(new XcodeBuildConfig(ConfigName, TargetName, FileReference.Combine(OutputDirectory, "Mac", MacExecutableName), FileReference.Combine(OutputDirectory, "IOS", IOSExecutableName), FileReference.Combine(OutputDirectory, "TVOS", TVOSExecutableName), ProjectTarget, Configuration));
                                             }
-                                            else if (BuildPlatform.Platform == UnrealTargetPlatform.IOS)
+                                            else if (BuildPlatform.Platform == UnrealTargetPlatform.IOS || BuildPlatform.Platform == UnrealTargetPlatform.TVOS)
                                             {
                                                 string IOSExecutableName = MakeExecutableFileName(ExeName, UnrealTargetPlatform.IOS, Configuration, ProjectTarget.TargetRules.Architecture, ProjectTarget.TargetRules.UndecoratedConfiguration);
                                                 string TVOSExecutableName = IOSExecutableName.Replace("-IOS-", "-TVOS-");
diff --git a/Engine/Source/Runtime/ApplicationCore/Private/IOS/IOSView.cpp b/Engine/Source/Runtime/ApplicationCore/Private/IOS/IOSView.cpp
--- a/Engine/Source/Runtime/ApplicationCore/Private/IOS/IOSView.cpp
+++ b/Engine/Source/Runtime/ApplicationCore/Private/IOS/IOSView.cpp
@@ -1105,20 +1105,23 @@ id<MTLDevice> GMetalDevice = nil;
 - (UIInterfaceOrientation)preferredInterfaceOrientationForPresentation
 {
 	FString PreferredLandscapeOrientation = "";
+	bool bSupportsLandscapeLeft = false;
+	bool bSupportsLandscapeRight = false;
 
+	GConfig->GetBool(TEXT("/Script/IOSRuntimeSettings.IOSRuntimeSettings"), TEXT("bSupportsLandscapeLeftOrientation"), bSupportsLandscapeLeft, GEngineIni);
+	GConfig->GetBool(TEXT("/Script/IOSRuntimeSettings.IOSRuntimeSettings"), TEXT("bSupportsLandscapeRightOrientation"), bSupportsLandscapeRight, GEngineIni);
 	GConfig->GetString(TEXT("/Script/IOSRuntimeSettings.IOSRuntimeSettings"), TEXT("PreferredLandscapeOrientation"), PreferredLandscapeOrientation, GEngineIni);
-	if (PreferredLandscapeOrientation.Equals("LandscapeLeft"))
+
+	if(bSupportsLandscapeLeft && bSupportsLandscapeRight)
 	{
+		if (PreferredLandscapeOrientation.Equals("LandscapeRight"))
+		{
+			return UIInterfaceOrientationLandscapeRight;
+		}
 		return UIInterfaceOrientationLandscapeLeft;
 	}
-	else if (PreferredLandscapeOrientation.Equals("LandscapeRight"))
-	{
-		return UIInterfaceOrientationLandscapeRight;
-	}
-	else
-	{
-		return UIInterfaceOrientationPortrait;
-	}
+
+	return UIInterfaceOrientationPortrait;
 }
 
 - (UIRectEdge)preferredScreenEdgesDeferringSystemGestures
diff --git a/Engine/Source/Runtime/ApplicationCore/Private/Windows/WindowsApplication.cpp b/Engine/Source/Runtime/ApplicationCore/Private/Windows/WindowsApplication.cpp
--- a/Engine/Source/Runtime/ApplicationCore/Private/Windows/WindowsApplication.cpp
+++ b/Engine/Source/Runtime/ApplicationCore/Private/Windows/WindowsApplication.cpp
@@ -466,7 +466,7 @@ FPlatformRect FWindowsApplication::GetWorkArea( const FPlatformRect& CurrentWind
  * @param OutHeight - Reference to output variable for monitor native height
  * @returns 'true' if data was extracted successfully, 'false' otherwise
  **/
-static bool GetMonitorSizeFromEDID(const HKEY hDevRegKey, int32& OutWidth, int32& OutHeight)
+static bool GetMonitorSizeFromEDID(const HKEY hDevRegKey, int32& OutWidth, int32& OutHeight, int32& OutDPI)
 {	
 	static const uint32 NameSize = 512;
 	static TCHAR ValueName[NameSize];
@@ -498,6 +498,25 @@ static bool GetMonitorSizeFromEDID(const HKEY hDevRegKey, int32& OutWidth, int32
 		OutWidth = ((EDIDData[DetailTimingDescriptorStartIndex+4] >> 4) << 8) | EDIDData[DetailTimingDescriptorStartIndex+2];
 		OutHeight = ((EDIDData[DetailTimingDescriptorStartIndex+7] >> 4) << 8) | EDIDData[DetailTimingDescriptorStartIndex+5];
 
+		const int32 HorizontalSizeOffset = 21;
+		const int32 VerticalSizeOffset = 22;
+		const float CmToInch = 0.393701f;
+
+		if (EDIDData[HorizontalSizeOffset] > 0 && EDIDData[VerticalSizeOffset] > 0)
+		{
+			float PhysicalWidth = CmToInch * (float)EDIDData[HorizontalSizeOffset];
+			float PhysicalHeight = CmToInch * (float)EDIDData[VerticalSizeOffset];
+
+			int32 HDpi = FMath::TruncToInt((float)OutWidth / PhysicalWidth);
+			int32 VDpi = FMath::TruncToInt((float)OutHeight / PhysicalHeight);
+
+			OutDPI = (HDpi + VDpi) / 2;
+		}
+		else
+		{
+			OutDPI = 0;
+		}
+
 		return true; // valid EDID found
 	}
 
@@ -511,7 +530,7 @@ static bool GetMonitorSizeFromEDID(const HKEY hDevRegKey, int32& OutWidth, int32
  * @praam OutHeight - Reference to output variable for monitor native height
  * @returns TRUE if data was extracted successfully, FALSE otherwise
  **/
-inline bool GetSizeForDevID(const FString& TargetDevID, int32& Width, int32& Height)
+inline bool GetSizeForDevID(const FString& TargetDevID, int32& Width, int32& Height, int32& DPI)
 {
 	static const GUID ClassMonitorGuid = {0x4d36e96e, 0xe325, 0x11ce, {0xbf, 0xc1, 0x08, 0x00, 0x2b, 0xe1, 0x03, 0x18}};
 
@@ -550,7 +569,7 @@ inline bool GetSizeForDevID(const FString& TargetDevID, int32& Width, int32& Hei
 
 					if (hDevRegKey && hDevRegKey != INVALID_HANDLE_VALUE)
 					{
-						bRes = GetMonitorSizeFromEDID(hDevRegKey, Width, Height);
+						bRes = GetMonitorSizeFromEDID(hDevRegKey, Width, Height, DPI);
 						RegCloseKey(hDevRegKey);
 						break;
 					}
@@ -628,12 +647,25 @@ static void GetMonitorsInfo(TArray<FMonitorInfo>& OutMonitorInfo)
 					Info.ID = FString::Printf(TEXT("%s"), Monitor.DeviceID);
 					Info.Name = Info.ID.Mid (8, Info.ID.Find (TEXT("\\"), ESearchCase::CaseSensitive, ESearchDir::FromStart, 9) - 8);
 
-					if (GetSizeForDevID(Info.Name, Info.NativeWidth, Info.NativeHeight))
+					if (GetSizeForDevID(Info.Name, Info.NativeWidth, Info.NativeHeight, Info.DPI))
 					{
 						Info.ID = Monitor.DeviceID;
 						Info.bIsPrimary = (DisplayDevice.StateFlags & DISPLAY_DEVICE_PRIMARY_DEVICE) > 0;
 
-						Info.DPI = FWindowsPlatformApplicationMisc::GetMonitorDPI(Info);
+						// sanity check for DPI values
+						if (Info.DPI < 96 || Info.DPI > 300)
+						{
+							// switch to default winapi value
+							Info.DPI = FWindowsPlatformApplicationMisc::GetMonitorDPI(Info);
+						}
+						else
+						{
+							// we also need to include the OS scaling value
+							const float CenterX = 0.5f * (Info.WorkArea.Right + Info.WorkArea.Left);
+							const float CenterY = 0.5f * (Info.WorkArea.Top + Info.WorkArea.Bottom);
+							const float DPIScaleFactor = FWindowsPlatformApplicationMisc::GetDPIScaleFactorAtPoint(CenterX, CenterY);
+							Info.DPI *= DPIScaleFactor;
+						}
 
 						OutMonitorInfo.Add(Info);
 
diff --git a/Engine/Source/Runtime/AudioMixer/Private/AudioMixerDevice.cpp b/Engine/Source/Runtime/AudioMixer/Private/AudioMixerDevice.cpp
--- a/Engine/Source/Runtime/AudioMixer/Private/AudioMixerDevice.cpp
+++ b/Engine/Source/Runtime/AudioMixer/Private/AudioMixerDevice.cpp
@@ -201,7 +201,7 @@ namespace Audio
 				AudioClockDelta = (double)OpenStreamParams.NumFrames / OpenStreamParams.SampleRate;
 
 				FAudioPluginInitializationParams PluginInitializationParams;
-				PluginInitializationParams.NumSources = MaxChannels;
+				PluginInitializationParams.NumSources = SourceManagerInitParams.NumSources;
 				PluginInitializationParams.SampleRate = SampleRate;
 				PluginInitializationParams.BufferLength = OpenStreamParams.NumFrames;
 				PluginInitializationParams.AudioDevicePtr = this;
diff --git a/Engine/Source/Runtime/Core/Private/Unix/UnixPlatformTime.cpp b/Engine/Source/Runtime/Core/Private/Unix/UnixPlatformTime.cpp
--- a/Engine/Source/Runtime/Core/Private/Unix/UnixPlatformTime.cpp
+++ b/Engine/Source/Runtime/Core/Private/Unix/UnixPlatformTime.cpp
@@ -88,14 +88,7 @@ uint64 FUnixTime::CallsPerSecondBenchmark(clockid_t BenchClockId, const char * B
 	char Buffer[256];
 	const uint64 kBenchmarkPeriodMicroSec = 1000000000ULL / 10;	// 0.1s
 
-	// basic sanity check
-	if (clock_getres(BenchClockId, nullptr) == -1)
-	{
-		FCStringAnsi::Snprintf(Buffer, sizeof(Buffer), "Clock_id %d (%s) is not supported on this system, clock_getres() fails.\n", BenchClockId, BenchClockIdName);
-		FCStringAnsi::Strncat(CalibrationLog, Buffer, sizeof(CalibrationLog));
-		return 0;	// unsupported clock id
-	}
-
+	// clock_getres() can fail when running on Windows Subsystem for Linux (but the clock can still be supported).
 	struct timespec ts;
 	if (clock_gettime(BenchClockId, &ts) == -1)
 	{
diff --git a/Engine/Source/Runtime/Core/Public/Android/AndroidMisc.h b/Engine/Source/Runtime/Core/Public/Android/AndroidMisc.h
--- a/Engine/Source/Runtime/Core/Public/Android/AndroidMisc.h
+++ b/Engine/Source/Runtime/Core/Public/Android/AndroidMisc.h
@@ -33,6 +33,11 @@ struct CORE_API FAndroidMisc : public FGenericPlatformMisc
 	static void PlatformTearDown();
 	static void PlatformHandleSplashScreen(bool ShowSplashScreen);
 
+	FORCEINLINE static int32 GetMaxPathLength()
+	{
+		return ANDROID_MAX_PATH;
+	}
+
 	DEPRECATED(4.21, "void FPlatformMisc::GetEnvironmentVariable(Name, Result, Length) is deprecated. Use FString FPlatformMisc::GetEnvironmentVariable(Name) instead.")
 	static void GetEnvironmentVariable(const TCHAR* VariableName, TCHAR* Result, int32 ResultLength);
 
diff --git a/Engine/Source/Runtime/Core/Public/Misc/Char.h b/Engine/Source/Runtime/Core/Public/Misc/Char.h
--- a/Engine/Source/Runtime/Core/Public/Misc/Char.h
+++ b/Engine/Source/Runtime/Core/Public/Misc/Char.h
@@ -137,7 +137,7 @@ typedef TChar<ANSICHAR> FCharAnsi;
 
 template <> inline TChar<WIDECHAR>::CharType TChar<WIDECHAR>::ToUpper(CharType Char)
 {
-#if PLATFORM_SWITCH && !PLATFORM_TCHAR_IS_4_BYTES
+#if !PLATFORM_TCHAR_IS_4_BYTES
 	// for 2 byte wide chars use ansi toupper() if it is not an actual wide char (value less than 128)
 	// towupper() is slow on some platforms.
 	if (Char < 128)
diff --git a/Engine/Source/Runtime/CoreUObject/Private/UObject/PropertyMap.cpp b/Engine/Source/Runtime/CoreUObject/Private/UObject/PropertyMap.cpp
--- a/Engine/Source/Runtime/CoreUObject/Private/UObject/PropertyMap.cpp
+++ b/Engine/Source/Runtime/CoreUObject/Private/UObject/PropertyMap.cpp
@@ -656,14 +656,10 @@ const TCHAR* UMapProperty::ImportText_Internal(const TCHAR* Buffer, void* Data,
 	}
 
 	uint8* TempPairStorage   = (uint8*)FMemory::Malloc(MapLayout.ValueOffset + ValueProp->ElementSize);
-	KeyProp  ->InitializeValue(TempPairStorage);
-	ValueProp->InitializeValue(TempPairStorage + MapLayout.ValueOffset);
 
 	bool bSuccess = false;
 	ON_SCOPE_EXIT
 	{
-		ValueProp->DestroyValue(TempPairStorage + MapLayout.ValueOffset);
-		KeyProp  ->DestroyValue(TempPairStorage);
 		FMemory::Free(TempPairStorage);
 
 		// If we are returning because of an error, remove any already-added elements from the map before returning
@@ -676,6 +672,14 @@ const TCHAR* UMapProperty::ImportText_Internal(const TCHAR* Buffer, void* Data,
 
 	for (;;)
 	{
+		KeyProp->InitializeValue(TempPairStorage);
+		ValueProp->InitializeValue(TempPairStorage + MapLayout.ValueOffset);
+		ON_SCOPE_EXIT
+		{
+			ValueProp->DestroyValue(TempPairStorage + MapLayout.ValueOffset);
+			KeyProp->DestroyValue(TempPairStorage);
+		};
+
 		if (*Buffer++ != TCHAR('('))
 		{
 			return nullptr;
diff --git a/Engine/Source/Runtime/Engine/Classes/Engine/DemoNetConnection.h b/Engine/Source/Runtime/Engine/Classes/Engine/DemoNetConnection.h
--- a/Engine/Source/Runtime/Engine/Classes/Engine/DemoNetConnection.h
+++ b/Engine/Source/Runtime/Engine/Classes/Engine/DemoNetConnection.h
@@ -26,6 +26,16 @@ struct FQueuedDemoPacket
 	uint32 SeenLevelIndex;
 
 public:
+	FORCEINLINE FQueuedDemoPacket(uint8* InData, int32 InSizeBytes, int32 InSizeBits) 
+		: Data()
+		, SizeBits(InSizeBits)
+		, Traits()
+		, SeenLevelIndex(0)
+	{
+		Data.AddUninitialized(InSizeBytes);
+		FMemory::Memcpy(Data.GetData(), InData, InSizeBytes);
+	}
+
 	FORCEINLINE FQueuedDemoPacket(uint8* InData, int32 InSizeBits, FOutPacketTraits& InTraits)
 		: Data()
 		, SizeBits(InSizeBits)
diff --git a/Engine/Source/Runtime/Engine/Classes/Engine/NetConnection.h b/Engine/Source/Runtime/Engine/Classes/Engine/NetConnection.h
--- a/Engine/Source/Runtime/Engine/Classes/Engine/NetConnection.h
+++ b/Engine/Source/Runtime/Engine/Classes/Engine/NetConnection.h
@@ -173,6 +173,17 @@ struct DelayedPacket
 	double SendTime;
 
 public:
+	DEPRECATED(4.21, "Use the constructor that takes PacketTraits for allowing for analytics and flags")
+	FORCEINLINE DelayedPacket(uint8* InData, int32 InSizeBytes, int32 InSizeBits)
+		: Data()
+		, SizeBits(InSizeBits)
+		, Traits()
+		, SendTime(0.0)
+	{
+		Data.AddUninitialized(InSizeBytes);
+		FMemory::Memcpy(Data.GetData(), InData, InSizeBytes);
+	}
+
 	FORCEINLINE DelayedPacket(uint8* InData, int32 InSizeBits, FOutPacketTraits& InTraits)
 		: Data()
 		, SizeBits(InSizeBits)
@@ -647,6 +658,13 @@ public:
 	/** Describe the connection. */
 	ENGINE_API virtual FString Describe();
 
+	DEPRECATED(4.21, "Use the method that allows for packet traits for analytics and modification")
+	ENGINE_API virtual void LowLevelSend(void* Data, int32 CountBytes, int32 CountBits)
+	{
+		FOutPacketTraits EmptyTraits;
+		LowLevelSend(Data, CountBits, EmptyTraits);
+	}
+
 	/**
 	 * Sends a byte stream to the remote endpoint using the underlying socket
 	 *
@@ -764,6 +782,11 @@ public:
 	 */
 	ENGINE_API virtual void InitConnection(UNetDriver* InDriver, EConnectionState InState, const FURL& InURL, int32 InConnectionSpeed=0, int32 InMaxPacket=0);
 
+	DEPRECATED(4.21, "Analytics providers are now handled in the NetDriver")
+	ENGINE_API virtual void InitHandler(TSharedPtr<IAnalyticsProvider> InProvider)
+	{
+		InitHandler();
+	}
 
 	/**
 	 * Initializes the PacketHandler
diff --git a/Engine/Source/Runtime/Engine/Classes/Engine/NetDriver.h b/Engine/Source/Runtime/Engine/Classes/Engine/NetDriver.h
--- a/Engine/Source/Runtime/Engine/Classes/Engine/NetDriver.h
+++ b/Engine/Source/Runtime/Engine/Classes/Engine/NetDriver.h
@@ -726,6 +726,13 @@ public:
 	/** PostTick actions */
 	ENGINE_API virtual void PostTickFlush();
 
+	DEPRECATED(4.21, "Please use the LowLevelSend that requires packet traits for analytics and packet modifiers.")
+	ENGINE_API virtual void LowLevelSend(FString Address, void* Data, int32 CountBits)
+	{
+		FOutPacketTraits EmptyTraits;
+		LowLevelSend(Address, Data, CountBits, EmptyTraits);
+	}
+
 	/**
 	 * Sends a 'connectionless' (not associated with a UNetConection) packet, to the specified address.
 	 * NOTE: Address is an abstract format defined by subclasses. Anything calling this, must use an address supplied by the net driver.
diff --git a/Engine/Source/Runtime/Engine/Private/NetConnection.cpp b/Engine/Source/Runtime/Engine/Private/NetConnection.cpp
--- a/Engine/Source/Runtime/Engine/Private/NetConnection.cpp
+++ b/Engine/Source/Runtime/Engine/Private/NetConnection.cpp
@@ -291,7 +291,7 @@ void UNetConnection::InitHandler()
 		{
 			Handler::Mode Mode = Driver->ServerConnection != nullptr ? Handler::Mode::Client : Handler::Mode::Server;
 
-			Handler->InitializeDelegates(FPacketHandlerLowLevelSend::CreateUObject(this, &UNetConnection::LowLevelSend));
+			Handler->InitializeDelegates(FPacketHandlerLowLevelSendTraits::CreateUObject(this, &UNetConnection::LowLevelSend));
 			Handler->NotifyAnalyticsProvider(Driver->AnalyticsProvider, Driver->AnalyticsAggregator);
 			Handler->Initialize(Mode, MaxPacket * 8, false);
 
diff --git a/Engine/Source/Runtime/Engine/Private/ShaderCompiler/ShaderCompilerXGE.cpp b/Engine/Source/Runtime/Engine/Private/ShaderCompiler/ShaderCompilerXGE.cpp
--- a/Engine/Source/Runtime/Engine/Private/ShaderCompiler/ShaderCompilerXGE.cpp
+++ b/Engine/Source/Runtime/Engine/Private/ShaderCompiler/ShaderCompilerXGE.cpp
@@ -81,6 +81,12 @@ namespace XGEShaderCompilerVariables
 				XGEShaderCompilerVariables::Enabled = 0;
 			}
 
+			//something going wrong with XGE shader compile on the build farm at the moment. Temporarily disable.
+			if (GIsBuildMachine)
+			{
+				XGEShaderCompilerVariables::Enabled = 0;
+			}
+
 			bInitialized = true;
 		}
 	}
diff --git a/Engine/Source/Runtime/IOS/IOSRuntimeSettings/Classes/IOSRuntimeSettings.h b/Engine/Source/Runtime/IOS/IOSRuntimeSettings/Classes/IOSRuntimeSettings.h
--- a/Engine/Source/Runtime/IOS/IOSRuntimeSettings/Classes/IOSRuntimeSettings.h
+++ b/Engine/Source/Runtime/IOS/IOSRuntimeSettings/Classes/IOSRuntimeSettings.h
@@ -350,7 +350,7 @@ public:
 	uint32 bSupportsLandscapeRightOrientation : 1;
 
 	// The Preferred Orientation will be used as the initial orientation at launch when both Landscape Left and Landscape Right orientations are to be supported.
-	UPROPERTY(GlobalConfig, EditAnywhere, Category = DeviceOrientations)
+	UPROPERTY(GlobalConfig, EditAnywhere, Category = DeviceOrientations, meta = (DisplayName = "Preferred Landscape Orientation"))
 	EIOSLandscapeOrientation PreferredLandscapeOrientation;
 
 	// Specifies the the display name for the application. This will be displayed under the icon on the device.
@@ -512,4 +512,8 @@ public:
 	virtual void PostInitProperties() override;
 	// End of UObject interface
 #endif
+
+private:
+	virtual void EnsureOrientationInProjectDefaultEngine();
+
 };
diff --git a/Engine/Source/Runtime/IOS/IOSRuntimeSettings/Private/IOSRuntimeSettings.cpp b/Engine/Source/Runtime/IOS/IOSRuntimeSettings/Private/IOSRuntimeSettings.cpp
--- a/Engine/Source/Runtime/IOS/IOSRuntimeSettings/Private/IOSRuntimeSettings.cpp
+++ b/Engine/Source/Runtime/IOS/IOSRuntimeSettings/Private/IOSRuntimeSettings.cpp
@@ -184,5 +184,16 @@ void UIOSRuntimeSettings::PostInitProperties()
 	{
 		MinimumiOSVersion = EIOSVersion::IOS_10;
 	}
+
+	EnsureOrientationInProjectDefaultEngine();
 }
 #endif
+
+void UIOSRuntimeSettings::EnsureOrientationInProjectDefaultEngine()
+{
+	// Make sure the values from BaseEngine.ini are written in the project's DefaultEngine.ini
+	UpdateSinglePropertyInConfigFile(GetClass()->FindPropertyByName(GET_MEMBER_NAME_CHECKED(UIOSRuntimeSettings, bSupportsPortraitOrientation)), GetDefaultConfigFilename());
+	UpdateSinglePropertyInConfigFile(GetClass()->FindPropertyByName(GET_MEMBER_NAME_CHECKED(UIOSRuntimeSettings, bSupportsUpsideDownOrientation)), GetDefaultConfigFilename());
+	UpdateSinglePropertyInConfigFile(GetClass()->FindPropertyByName(GET_MEMBER_NAME_CHECKED(UIOSRuntimeSettings, bSupportsLandscapeLeftOrientation)), GetDefaultConfigFilename());
+	UpdateSinglePropertyInConfigFile(GetClass()->FindPropertyByName(GET_MEMBER_NAME_CHECKED(UIOSRuntimeSettings, PreferredLandscapeOrientation)), GetDefaultConfigFilename());
+}
diff --git a/Engine/Source/Runtime/Landscape/Private/Landscape.cpp b/Engine/Source/Runtime/Landscape/Private/Landscape.cpp
--- a/Engine/Source/Runtime/Landscape/Private/Landscape.cpp
+++ b/Engine/Source/Runtime/Landscape/Private/Landscape.cpp
@@ -1051,6 +1051,14 @@ void ULandscapeComponent::GetGeneratedTexturesAndMaterialInstances(TArray<UObjec
 			}
 		}
 	}
+
+	for (UMaterialInstanceConstant* MaterialInstance : MobileCombinationMaterialInstances)
+	{
+		for (ULandscapeMaterialInstanceConstant* CurrentMIC = Cast<ULandscapeMaterialInstanceConstant>(MaterialInstance); CurrentMIC; CurrentMIC = Cast<ULandscapeMaterialInstanceConstant>(CurrentMIC->Parent))
+		{
+			OutTexturesAndMaterials.Add(CurrentMIC);
+		}
+	}	
 }
 #endif
 
diff --git a/Engine/Source/Runtime/PacketHandlers/PacketHandler/Private/PacketHandler.cpp b/Engine/Source/Runtime/PacketHandlers/PacketHandler/Private/PacketHandler.cpp
--- a/Engine/Source/Runtime/PacketHandlers/PacketHandler/Private/PacketHandler.cpp
+++ b/Engine/Source/Runtime/PacketHandlers/PacketHandler/Private/PacketHandler.cpp
@@ -33,6 +33,7 @@ PacketHandler::PacketHandler(FDDoSDetection* InDDoS/*=nullptr*/)
 	, bConnectionlessHandler(false)
 	, DDoS(InDDoS)
 	, LowLevelSendDel()
+	, LowLevelSendDel_Deprecated()
 	, HandshakeCompleteDel()
 	, OutgoingPacket()
 	, IncomingPacket()
@@ -605,7 +606,7 @@ void PacketHandler::SendHandlerPacket(HandlerComponent* InComponent, FBitWriter&
 	// Prevent any cases where a send happens before the handler is ready.
 	check(State != Handler::State::Uninitialized);
 
-	if (LowLevelSendDel.IsBound())
+	if (LowLevelSendDel.IsBound() || LowLevelSendDel_Deprecated.IsBound())
 	{
 		bool bEncounteredComponent = false;
 
@@ -661,7 +662,8 @@ void PacketHandler::SendHandlerPacket(HandlerComponent* InComponent, FBitWriter&
 
 			bRawSend = true;
 
-			LowLevelSendDel.Execute(Writer.GetData(), Writer.GetNumBits(), Traits);
+			LowLevelSendDel_Deprecated.ExecuteIfBound(Writer.GetData(), Writer.GetNumBytes(), Writer.GetNumBits());
+			LowLevelSendDel.ExecuteIfBound(Writer.GetData(), Writer.GetNumBits(), Traits);
 
 			bRawSend = bOldRawSend;
 		}
diff --git a/Engine/Source/Runtime/PacketHandlers/PacketHandler/Public/PacketHandler.h b/Engine/Source/Runtime/PacketHandlers/PacketHandler/Public/PacketHandler.h
--- a/Engine/Source/Runtime/PacketHandlers/PacketHandler/Public/PacketHandler.h
+++ b/Engine/Source/Runtime/PacketHandlers/PacketHandler/Public/PacketHandler.h
@@ -26,7 +26,9 @@ class FNetAnalyticsAggregator;
  */
 
 // Delegate for allowing access to LowLevelSend, without a dependency upon Engine
-DECLARE_DELEGATE_ThreeParams(FPacketHandlerLowLevelSend, void* /* Data */, int32 /* CountBits */, FOutPacketTraits& /* Traits */);
+DECLARE_DELEGATE_ThreeParams(FPacketHandlerLowLevelSendTraits, void* /* Data */, int32 /* CountBits */, FOutPacketTraits& /* Traits */);
+
+DECLARE_DELEGATE_ThreeParams(FPacketHandlerLowLevelSend, void* /* Data */, int32 /* CountBytes */, int32 /* CountBits */);
 
 /**
  * Callback for notifying higher-level code that handshaking has completed, and that packets are now ready to send without buffering
@@ -149,6 +151,21 @@ private:
 	}
 
 public:
+	DEPRECATED(4.21, "Please use the new constructor that adds support for analytics and better precision")
+	BufferedPacket(uint8* InCopyData, uint32 InCountBits, float InResendTime=0.f, uint32 InId=0)
+		: CountBits(InCountBits)
+		, Traits()
+		, ResendTime(double(InResendTime))
+		, Id(InId)
+		, Address()
+		, FromComponent(nullptr)
+	{
+		check(InCopyData != nullptr);
+
+		Data = new uint8[FMath::DivideAndRoundUp(InCountBits, 8u)];
+		FMemory::Memcpy(Data, InCopyData, FMath::DivideAndRoundUp(InCountBits, 8u));
+	}
+
 	BufferedPacket(uint8* InCopyData, uint32 InCountBits, FOutPacketTraits& InTraits, double InResendTime=0.0, uint32 InId=0)
 		: CountBits(InCountBits)
 		, Traits(InTraits)
@@ -163,6 +180,23 @@ public:
 		FMemory::Memcpy(Data, InCopyData, FMath::DivideAndRoundUp(InCountBits, 8u));
 	}
 
+	DEPRECATED(4.21, "Please use the new constructor that adds support for analytics")
+	BufferedPacket(const FString& InAddress, uint8* InCopyData, uint32 InCountBits, double InResendTime=0.0, uint32 InId=0)
+		: CountBits(InCountBits)
+		, Traits()
+		, ResendTime(InResendTime)
+		, Id(InId)
+		, Address()
+		, FromComponent(nullptr)
+	{
+		check(InCopyData != nullptr);
+
+		Data = new uint8[FMath::DivideAndRoundUp(InCountBits, 8u)];
+		FMemory::Memcpy(Data, InCopyData, FMath::DivideAndRoundUp(InCountBits, 8u));
+
+		Address = InAddress;
+	}
+
 	BufferedPacket(const FString& InAddress, uint8* InCopyData, uint32 InCountBits, FOutPacketTraits& InTraits, double InResendTime=0.0, uint32 InId=0)
 		: BufferedPacket(InCopyData, InCountBits, InTraits, InResendTime, InId)
 	{
@@ -203,12 +237,18 @@ public:
 	void Initialize(Handler::Mode Mode, uint32 InMaxPacketBits, bool bConnectionlessOnly=false,
 					TSharedPtr<class IAnalyticsProvider> InProvider=nullptr, FDDoSDetection* InDDoS=nullptr);
 
+	DEPRECATED(4.21, "Use the traits based delegate instead for compatibility with other systems.")
+	void InitializeDelegates(FPacketHandlerLowLevelSend InLowLevelSendDel)
+	{
+		LowLevelSendDel_Deprecated = InLowLevelSendDel;
+	}
+
 	/**
 	 * Used for external initialization of delegates
 	 *
 	 * @param InLowLevelSendDel		The delegate the PacketHandler should use for triggering packet sends
 	 */
-	void InitializeDelegates(FPacketHandlerLowLevelSend InLowLevelSendDel)
+	void InitializeDelegates(FPacketHandlerLowLevelSendTraits InLowLevelSendDel)
 	{
 		LowLevelSendDel = InLowLevelSendDel;
 	}
@@ -298,6 +338,13 @@ public:
 		return Incoming_Internal(Packet, CountBytes, false, EmptyString);
 	}
 
+	DEPRECATED(4.21, "Please move to the functional flow that includes support for PacketTraits.")
+	FORCEINLINE const ProcessedPacket Outgoing(uint8* Packet, int32 CountBits)
+	{
+		FOutPacketTraits EmptyTraits;
+		return Outgoing(Packet, CountBits, EmptyTraits);
+	}
+
 	/**
 	 * Processes outgoing packets at the PacketHandler level, after all UNetConnection processing.
 	 *
@@ -330,6 +377,13 @@ public:
 		return Incoming_Internal(Packet, CountBytes, true, Address);
 	}
 
+	DEPRECATED(4.21, "Please use the member that supports PacketTraits for alllowing additional flags on sends.")
+	FORCEINLINE const ProcessedPacket OutgoingConnectionless(const FString& Address, uint8* Packet, int32 CountBits)
+	{
+		FOutPacketTraits EmptyTraits;
+		return OutgoingConnectionless(Address, Packet, CountBits, EmptyTraits);
+	}
+
 	/**
 	 * Processes outgoing packets without a UNetConnection, in the same manner as 'Outgoing' above
 	 * NOTE: Connectionless packets are unreliable.
@@ -376,6 +430,14 @@ protected:
 	const ProcessedPacket Outgoing_Internal(uint8* Packet, int32 CountBits, FOutPacketTraits& Traits, bool bConnectionless, const FString& Address );
 
 public:
+
+	DEPRECATED(4.21, "Please use the packet traits when sending to handle modifications of packets and analytics.")
+	void SendHandlerPacket(HandlerComponent* InComponent, FBitWriter& Writer)
+	{
+		FOutPacketTraits EmptyTraits;
+		SendHandlerPacket(InComponent, Writer, EmptyTraits);
+	}
+
 	/**
 	 * Send a packet originating from a HandlerComponent - will process through the HandlerComponents chain,
 	 * starting after the triggering component.
@@ -526,7 +588,10 @@ private:
 	FDDoSDetection* DDoS;
 
 	/** Delegate used for triggering PacketHandler/HandlerComponent-sourced sends */
-	FPacketHandlerLowLevelSend LowLevelSendDel;
+	FPacketHandlerLowLevelSendTraits LowLevelSendDel;
+
+	/** Delegate used for triggering PacketHandler/HandlerComponent-sourced sends (DEPRECATED) */
+	FPacketHandlerLowLevelSend LowLevelSendDel_Deprecated;
 
 	/** Delegate used for notifying that handshaking has completed */
 	FPacketHandlerHandshakeComplete HandshakeCompleteDel;
@@ -646,6 +711,13 @@ public:
 	 */
 	virtual void Incoming(FBitReader& Packet) = 0;
 
+	DEPRECATED(4.21, "Use the other Outgoing function as it allows for packet modifiers and traits.")
+	virtual void Outgoing(FBitWriter& Packet)
+	{
+		FOutPacketTraits EmptyTraits;
+		Outgoing(Packet, EmptyTraits);
+	}
+
 	/**
 	 * Handles any outgoing packets
 	 *
@@ -662,6 +734,13 @@ public:
 	 */
 	virtual void IncomingConnectionless(const FString& Address, FBitReader& Packet) = 0;
 
+	DEPRECATED(4.21, "Use the method that allows traits on the packet.")
+	virtual void OutgoingConnectionless(const FString& Address, FBitWriter& Packet)
+	{
+		FOutPacketTraits EmptyTraits;
+		OutgoingConnectionless(Address, Packet, EmptyTraits);
+	}
+
 	/**
 	 * Handles any outgoing packets not associated with a UNetConnection
 	 *
@@ -721,6 +800,11 @@ public:
 	/** Returns the name of this component. */
 	FName GetName() const { return Name; }
 
+	DEPRECATED(4.21, "The Analytics Provider is now handled in the main PacketHandler class.")
+	virtual void SetAnalyticsProvider(TSharedPtr<class IAnalyticsProvider> Provider)
+	{
+	}
+
 	/**
 	 * Notification that the analytics provider has been updated
 	 * NOTE: Can also mean disabled, e.g. during hotfix
diff --git a/Engine/Source/Runtime/PacketHandlers/ReliabilityHandlerComponent/Public/ReliabilityHandlerComponent.h b/Engine/Source/Runtime/PacketHandlers/ReliabilityHandlerComponent/Public/ReliabilityHandlerComponent.h
--- a/Engine/Source/Runtime/PacketHandlers/ReliabilityHandlerComponent/Public/ReliabilityHandlerComponent.h
+++ b/Engine/Source/Runtime/PacketHandlers/ReliabilityHandlerComponent/Public/ReliabilityHandlerComponent.h
@@ -28,6 +28,13 @@ public:
 	/* Queues a packet for resending */
 	void QueuePacketForResending(uint8* Packet, int32 CountBits, FOutPacketTraits& Traits);
 
+	DEPRECATED(4.21, "Use the PacketTraits version for sending packets with additional flags and options")
+	FORCEINLINE void QueueHandlerPacketForResending(HandlerComponent* InComponent, uint8* Packet, int32 CountBits)
+	{
+		FOutPacketTraits EmptyTraits;
+		QueueHandlerPacketForResending(InComponent, Packet, CountBits, EmptyTraits);
+	}
+
 	/**
 	 * Queues a packet sent through SendHandlerPacket, for resending
 	 *
diff --git a/Engine/Source/Runtime/Renderer/Private/PostProcess/DiaphragmDOF.h b/Engine/Source/Runtime/Renderer/Private/PostProcess/DiaphragmDOF.h
--- a/Engine/Source/Runtime/Renderer/Private/PostProcess/DiaphragmDOF.h
+++ b/Engine/Source/Runtime/Renderer/Private/PostProcess/DiaphragmDOF.h
@@ -123,8 +123,7 @@ inline bool IsSupported(EShaderPlatform ShaderPlatform)
 		ShaderPlatform == SP_PS4 ||
 		IsVulkanSM5Platform(ShaderPlatform) ||
 		IsMetalSM5Platform(ShaderPlatform) ||
-		ShaderPlatform == SP_SWITCH ||
-		ShaderPlatform == SP_SWITCH_FORWARD;
+		ShaderPlatform == SP_SWITCH;
 }
 
 
diff --git a/Engine/Source/Runtime/Renderer/Private/ShadowRendering.cpp b/Engine/Source/Runtime/Renderer/Private/ShadowRendering.cpp
--- a/Engine/Source/Runtime/Renderer/Private/ShadowRendering.cpp
+++ b/Engine/Source/Runtime/Renderer/Private/ShadowRendering.cpp
@@ -762,7 +762,7 @@ void FProjectedShadowInfo::SetupProjectionStencilMask(
 
 		RHIUnlockVertexBuffer(VertexBufferRHI);
 		RHICmdList.SetStreamSource(0, VertexBufferRHI, 0);
-		RHICmdList.DrawPrimitive(PT_TriangleStrip, 0, (CascadeSettings.ShadowSplitIndex > 0) ? 4 : 2, 1);
+		RHICmdList.DrawPrimitive(PT_TriangleList, 0, (CascadeSettings.ShadowSplitIndex > 0) ? 4 : 2, 1);
 	}
 	// Not a preshadow, mask the projection to any pixels inside the frustum.
 	else
diff --git a/Engine/Source/Runtime/Sockets/Private/BSDSockets/SocketsBSD.cpp b/Engine/Source/Runtime/Sockets/Private/BSDSockets/SocketsBSD.cpp
--- a/Engine/Source/Runtime/Sockets/Private/BSDSockets/SocketsBSD.cpp
+++ b/Engine/Source/Runtime/Sockets/Private/BSDSockets/SocketsBSD.cpp
@@ -423,8 +423,7 @@ bool FSocketBSD::LeaveMulticastGroup(const FInternetAddr& GroupAddress)
 bool FSocketBSD::SetMulticastLoopback(bool bLoopback)
 {
 #if PLATFORM_HAS_BSD_IPV6_SOCKETS
-	bool IPv6Loop = bLoopback;
-	setsockopt(Socket, IPPROTO_IPV6, IPV6_MULTICAST_LOOP, (char*)&IPv6Loop, sizeof(IPv6Loop));
+	return (setsockopt(Socket, IPPROTO_IPV6, IP_MULTICAST_LOOP, (char*)&bLoopback, sizeof(bLoopback)) == 0);
 #endif
 
 	return (setsockopt(Socket, IPPROTO_IP, IP_MULTICAST_LOOP, (char*)&bLoopback, sizeof(bLoopback)) == 0);
@@ -434,8 +433,7 @@ bool FSocketBSD::SetMulticastLoopback(bool bLoopback)
 bool FSocketBSD::SetMulticastTtl(uint8 TimeToLive)
 {
 #if PLATFORM_HAS_BSD_IPV6_SOCKETS
-	uint8 IPv6TTL = TimeToLive;
-	setsockopt(Socket, IPPROTO_IPV6, IP_MULTICAST_TTL, (char*)&IPv6TTL, sizeof(IPv6TTL));
+	return (setsockopt(Socket, IPPROTO_IPV6, IP_MULTICAST_TTL, (char*)&TimeToLive, sizeof(TimeToLive)) == 0);
 #endif
 
 	return (setsockopt(Socket, IPPROTO_IP, IP_MULTICAST_TTL, (char*)&TimeToLive, sizeof(TimeToLive)) == 0);
diff --git a/Engine/Source/Runtime/UMG/Public/Components/ListViewBase.h b/Engine/Source/Runtime/UMG/Public/Components/ListViewBase.h
--- a/Engine/Source/Runtime/UMG/Public/Components/ListViewBase.h
+++ b/Engine/Source/Runtime/UMG/Public/Components/ListViewBase.h
@@ -321,7 +321,9 @@ private:
 
 		UUserWidget& EntryWidget = OnGenerateEntryWidgetInternal(Item, DesiredEntryClass, OwnerTable);
 		EntryWidget.SetPadding(GetDesiredEntryPadding(Item));
-		return StaticCastSharedPtr<SObjectTableRow<ItemType>>(EntryWidget.GetCachedWidget()).ToSharedRef();
+		TSharedPtr<SWidget> CachedWidget = EntryWidget.GetCachedWidget();
+		CachedWidget->SetCanTick(true); // this is a hack to force ticking to true so selection works (which should NOT require ticking! but currently does)
+		return StaticCastSharedPtr<SObjectTableRow<ItemType>>(CachedWidget).ToSharedRef();
 	}
 
 	void HandleItemClicked(ItemType Item)
diff --git a/Engine/Source/Runtime/UMG/Public/Slate/SObjectTableRow.h b/Engine/Source/Runtime/UMG/Public/Slate/SObjectTableRow.h
--- a/Engine/Source/Runtime/UMG/Public/Slate/SObjectTableRow.h
+++ b/Engine/Source/Runtime/UMG/Public/Slate/SObjectTableRow.h
@@ -79,8 +79,6 @@ public:
 			[
 				ContentWidget.ToSharedRef()
 			], &InWidgetObject);
-
-		SetCanTick(true);
 	}
 
 	virtual ~SObjectTableRow()
