PacketHandler component profiles per netdriver.

diff --git a/Engine/Plugins/Online/OnlineSubsystemSteam/Source/Private/OnlineAuthInterfaceSteam.cpp b/Engine/Plugins/Online/OnlineSubsystemSteam/Source/Private/OnlineAuthInterfaceSteam.cpp
--- a/Engine/Plugins/Online/OnlineSubsystemSteam/Source/Private/OnlineAuthInterfaceSteam.cpp
+++ b/Engine/Plugins/Online/OnlineSubsystemSteam/Source/Private/OnlineAuthInterfaceSteam.cpp
@@ -15,6 +15,7 @@
 
 // Determining if we need to be disabled or not
 #include "Misc/ConfigCacheIni.h"
+#include "PacketHandler.h"
 
 // Steam tells us this number in documentation, however there's no define within the SDK
 #define STEAM_AUTH_MAX_TICKET_LENGTH_IN_BYTES 1024
@@ -32,18 +33,32 @@ FOnlineAuthSteam::FOnlineAuthSteam(FOnlineSubsystemSteam* InSubsystem) :
 	bNeverSendKey(false),
 	bSendBadId(false)
 {
-	TArray<FString> ComponentList;
-	GConfig->GetArray(TEXT("PacketHandlerComponents"), TEXT("Components"), ComponentList, GEngineIni);
-
-	for (FString CompStr : ComponentList)
+	const FString SteamModuleName(TEXT("SteamAuthComponentModuleInterface"));
+	if (!PacketHandler::DoesAnyProfileHaveComponent(SteamModuleName))
 	{
-		if (CompStr.Contains(TEXT("SteamAuthComponentModuleInterface")))
+		// Pull the components to see if there's anything we can use.
+		TArray<FString> ComponentList;
+		GConfig->GetArray(TEXT("PacketHandlerComponents"), TEXT("Components"), ComponentList, GEngineIni);
+
+		// Check if Steam Auth is enabled anywhere.
+		for (FString CompStr : ComponentList)
 		{
-			UE_LOG_ONLINE(Log, TEXT("AUTH: Steam Auth Enabled"));
-			bEnabled = true;
-			return;
+			if (CompStr.Contains(SteamModuleName))
+			{
+				bEnabled = true;
+				break;
+			}
 		}
 	}
+	else
+	{
+		bEnabled = true;
+	}
+
+	if (bEnabled)
+	{
+		UE_LOG_ONLINE(Log, TEXT("AUTH: Steam Auth Enabled"));
+	}
 }
 
 FOnlineAuthSteam::FOnlineAuthSteam() :
diff --git a/Engine/Source/Runtime/Engine/Private/NetConnection.cpp b/Engine/Source/Runtime/Engine/Private/NetConnection.cpp
--- a/Engine/Source/Runtime/Engine/Private/NetConnection.cpp
+++ b/Engine/Source/Runtime/Engine/Private/NetConnection.cpp
@@ -66,11 +66,11 @@ UNetConnection* UNetConnection::GNetConnectionBeingCleanedUp = NULL;
 
 UNetConnection::UNetConnection(const FObjectInitializer& ObjectInitializer)
 :	UPlayer(ObjectInitializer)
-,	Driver				( NULL )
+,	Driver				( nullptr )
 ,	PackageMapClass		( UPackageMapClient::StaticClass() )
-,	PackageMap			( NULL )
-,	ViewTarget			( NULL )
-,   OwningActor			( NULL )
+,	PackageMap			( nullptr )
+,	ViewTarget			( nullptr )
+,   OwningActor			( nullptr )
 ,	MaxPacket			( 0 )
 ,	InternalAck			( false )
 ,	MaxPacketHandlerBits ( 0 )
@@ -86,7 +86,7 @@ UNetConnection::UNetConnection(const FObjectInitializer& ObjectInitializer)
 
 ,	AllowMerge			( false )
 ,	TimeSensitive		( false )
-,	LastOutBunch		( NULL )
+,	LastOutBunch		( nullptr )
 ,	SendBunchHeader		( MAX_BUNCH_HEADER_BITS )
 
 ,	StatPeriod			( 1.f  )
@@ -297,7 +297,7 @@ void UNetConnection::InitHandler()
 
 			Handler->InitializeDelegates(FPacketHandlerLowLevelSendTraits::CreateUObject(this, &UNetConnection::LowLevelSend));
 			Handler->NotifyAnalyticsProvider(Driver->AnalyticsProvider, Driver->AnalyticsAggregator);
-			Handler->Initialize(Mode, MaxPacket * 8, false);
+			Handler->Initialize(Mode, MaxPacket * 8, false, nullptr, nullptr, Driver->NetDriverName);
 
 
 			// Add handling for the stateless connect handshake, for connectionless packets, as the outermost layer
diff --git a/Engine/Source/Runtime/Engine/Private/NetworkDriver.cpp b/Engine/Source/Runtime/Engine/Private/NetworkDriver.cpp
--- a/Engine/Source/Runtime/Engine/Private/NetworkDriver.cpp
+++ b/Engine/Source/Runtime/Engine/Private/NetworkDriver.cpp
@@ -1394,7 +1394,7 @@ void UNetDriver::InitConnectionlessHandler()
 		if (ConnectionlessHandler.IsValid())
 		{
 			ConnectionlessHandler->NotifyAnalyticsProvider(AnalyticsProvider, AnalyticsAggregator);
-			ConnectionlessHandler->Initialize(Handler::Mode::Server, MAX_PACKET_SIZE, true);
+			ConnectionlessHandler->Initialize(Handler::Mode::Server, MAX_PACKET_SIZE, true, nullptr, nullptr, NetDriverName);
 
 			// Add handling for the stateless connect handshake, for connectionless packets, as the outermost layer
 			TSharedPtr<HandlerComponent> NewComponent =
diff --git a/Engine/Source/Runtime/PacketHandlers/PacketHandler/Classes/PacketHandlerProfileConfig.h b/Engine/Source/Runtime/PacketHandlers/PacketHandler/Classes/PacketHandlerProfileConfig.h
new file mode 100644
--- /dev/null
+++ b/Engine/Source/Runtime/PacketHandlers/PacketHandler/Classes/PacketHandlerProfileConfig.h
@@ -0,0 +1,22 @@
+// Copyright 1998-2018 Epic Games, Inc. All Rights Reserved.
+
+#pragma once
+
+#include "CoreMinimal.h"
+#include "UObject/ObjectMacros.h"
+#include "UObject/Object.h"
+#include "PacketHandlerProfileConfig.generated.h"
+
+/**
+ * PacketHandler Configuration Profiles based on each NetDriver
+ */
+UCLASS(config=Engine, PerObjectConfig)
+class UPacketHandlerProfileConfig : public UObject
+{
+	GENERATED_UCLASS_BODY()
+
+public:
+
+	UPROPERTY(config)
+	TArray<FString> Components;
+};
diff --git a/Engine/Source/Runtime/PacketHandlers/PacketHandler/Private/PacketHandler.cpp b/Engine/Source/Runtime/PacketHandlers/PacketHandler/Private/PacketHandler.cpp
--- a/Engine/Source/Runtime/PacketHandlers/PacketHandler/Private/PacketHandler.cpp
+++ b/Engine/Source/Runtime/PacketHandlers/PacketHandler/Private/PacketHandler.cpp
@@ -13,6 +13,7 @@
 #include "DDoSDetection.h"
 #include "HandlerComponentFactory.h"
 #include "ReliabilityHandlerComponent.h"
+#include "PacketHandlerProfileConfig.h"
 
 // @todo #JohnB: There is quite a lot of inefficient copying of packet data going on.
 //					Redo the whole packet parsing/modification pipeline.
@@ -82,7 +83,7 @@ void PacketHandler::Tick(float DeltaTime)
 }
 
 void PacketHandler::Initialize(Handler::Mode InMode, uint32 InMaxPacketBits, bool bConnectionlessOnly/*=false*/,
-								TSharedPtr<IAnalyticsProvider> InProvider/*=nullptr*/, FDDoSDetection* InDDoS/*=nullptr*/)
+								TSharedPtr<IAnalyticsProvider> InProvider/*=nullptr*/, FDDoSDetection* InDDoS/*=nullptr*/, FName InDriverProfile/*=NAME_None*/)
 {
 	Mode = InMode;
 	MaxPacketBits = InMaxPacketBits;
@@ -96,8 +97,14 @@ void PacketHandler::Initialize(Handler::Mode InMode, uint32 InMaxPacketBits, boo
 	if (!bConnectionlessHandler)
 	{
 		TArray<FString> Components;
-
-		GConfig->GetArray(TEXT("PacketHandlerComponents"), TEXT("Components"), Components, GEngineIni);
+		UPacketHandlerProfileConfig* CurNetDriverProfile = NewObject<UPacketHandlerProfileConfig>((UObject*)GetTransientPackage(), InDriverProfile);
+		Components.Append(CurNetDriverProfile->Components);
+		
+		// If we didn't get any matches, push in the regular components.
+		if (Components.Num() == 0)
+		{
+			GConfig->GetArray(TEXT("PacketHandlerComponents"), TEXT("Components"), Components, GEngineIni);
+		}
 
 		for (const FString& CurComponent : Components)
 		{
@@ -790,6 +797,49 @@ void PacketHandler::HandlerComponentInitialized(HandlerComponent* InComponent)
 	}
 }
 
+bool PacketHandler::DoesAnyProfileHaveComponent(const FString& InComponentName)
+{
+	TArray<FString> ProfileSectionNames;
+	if (GConfig->GetPerObjectConfigSections(GEngineIni, TEXT("PacketHandlerProfileConfig"), ProfileSectionNames))
+	{
+		for (const FString& CurProfileSection : ProfileSectionNames)
+		{
+			FName CurNetDriver(*CurProfileSection.Left(CurProfileSection.Find(TEXT(" "))));
+			if (DoesProfileHaveComponent(CurNetDriver, InComponentName))
+			{
+				return true;
+			}
+		}
+	}
+
+	return false;
+}
+
+bool PacketHandler::DoesProfileHaveComponent(const FName InNetDriverName, const FString& InComponentName)
+{
+	TArray<FString> Components;
+	// If this call was made during a module load, we cannot use the packethandler config object
+	if (GetTransientPackage() == nullptr)
+	{
+		FString DriverProfileCategory = FString::Printf(TEXT("%s PacketHandlerProfileConfig"), *InNetDriverName.ToString());
+		GConfig->GetArray(*DriverProfileCategory, TEXT("Components"), Components, GEngineIni);
+	}
+	else
+	{
+		UPacketHandlerProfileConfig* CurNetDriverProfile = NewObject<UPacketHandlerProfileConfig>((UObject*)GetTransientPackage(), InNetDriverName);
+		Components.Append(CurNetDriverProfile->Components);
+	}
+	
+	for (const FString& Component : Components)
+	{
+		if (Component.Contains(InComponentName, ESearchCase::CaseSensitive))
+		{
+			return true;
+		}
+	}
+	return false;
+}
+
 BufferedPacket* PacketHandler::GetQueuedPacket()
 {
 	BufferedPacket* QueuedPacket = nullptr;
diff --git a/Engine/Source/Runtime/PacketHandlers/PacketHandler/Private/PacketHandlerProfileConfig.cpp b/Engine/Source/Runtime/PacketHandlers/PacketHandler/Private/PacketHandlerProfileConfig.cpp
new file mode 100644
--- /dev/null
+++ b/Engine/Source/Runtime/PacketHandlers/PacketHandler/Private/PacketHandlerProfileConfig.cpp
@@ -0,0 +1,8 @@
+// Copyright 1998-2018 Epic Games, Inc. All Rights Reserved.
+
+#include "PacketHandlerProfileConfig.h"
+
+UPacketHandlerProfileConfig::UPacketHandlerProfileConfig(const FObjectInitializer& ObjectInitializer)
+	: Super(ObjectInitializer)
+{
+}
diff --git a/Engine/Source/Runtime/PacketHandlers/PacketHandler/Public/PacketHandler.h b/Engine/Source/Runtime/PacketHandlers/PacketHandler/Public/PacketHandler.h
--- a/Engine/Source/Runtime/PacketHandlers/PacketHandler/Public/PacketHandler.h
+++ b/Engine/Source/Runtime/PacketHandlers/PacketHandler/Public/PacketHandler.h
@@ -233,9 +233,10 @@ public:
 	 * @param bConnectionlessOnly	Whether or not this is a connectionless-only manager (ignores .ini components)
 	 * @param InProvider			The analytics provider
 	 * @param InDDoS				Reference to the owning net drivers DDoS detection handler
+	 * @param InDriverProfile		The PacketHandler configuration profile to use
 	 */
 	void Initialize(Handler::Mode Mode, uint32 InMaxPacketBits, bool bConnectionlessOnly=false,
-					TSharedPtr<class IAnalyticsProvider> InProvider=nullptr, FDDoSDetection* InDDoS=nullptr);
+					TSharedPtr<class IAnalyticsProvider> InProvider=nullptr, FDDoSDetection* InDDoS=nullptr, FName InDriverProfile=NAME_None);
 
 	UE_DEPRECATED(4.21, "Use the traits based delegate instead for compatibility with other systems.")
 	void InitializeDelegates(FPacketHandlerLowLevelSend InLowLevelSendDel)
@@ -475,6 +476,23 @@ public:
 		QueuedHandlerPackets.Enqueue(PacketToQueue);
 	}
 
+	/**
+	 * Searches the PacketHandler profile configurations to find if a component is listed.
+	 *
+	 * @param InComponentName	The PacketHandler Component to search for
+	 * @return if there is a profile that has the component included.
+	 */
+	static bool DoesAnyProfileHaveComponent(const FString& InComponentName);
+
+	/**
+	 * Searches the PacketHandler profile configuration for the given netdriver to find if a component is listed.
+	 *
+	 * @param InNetDriverName	The name of the netdriver to search configuration for
+	 * @param InComponentName	The component to search for
+	 * @return if the component is listed in the profile configuration.
+	 */
+	static bool DoesProfileHaveComponent(const FName InNetDriverName, const FString& InComponentName);
+
 	/**
 	 * Gets a packet from the buffered packet queue for sending
 	 *
