Setting copyright to 2019.

#rnx

diff --git a/Engine/Source/Runtime/PacketHandlers/PacketHandler/Classes/PacketHandlerProfileConfig.h b/Engine/Source/Runtime/PacketHandlers/PacketHandler/Classes/PacketHandlerProfileConfig.h
--- a/Engine/Source/Runtime/PacketHandlers/PacketHandler/Classes/PacketHandlerProfileConfig.h
+++ b/Engine/Source/Runtime/PacketHandlers/PacketHandler/Classes/PacketHandlerProfileConfig.h
@@ -1,4 +1,4 @@
-// Copyright 1998-2018 Epic Games, Inc. All Rights Reserved.
+// Copyright 1998-2019 Epic Games, Inc. All Rights Reserved.
 
 #pragma once
 
diff --git a/Engine/Source/Runtime/PacketHandlers/PacketHandler/Private/PacketHandlerProfileConfig.cpp b/Engine/Source/Runtime/PacketHandlers/PacketHandler/Private/PacketHandlerProfileConfig.cpp
--- a/Engine/Source/Runtime/PacketHandlers/PacketHandler/Private/PacketHandlerProfileConfig.cpp
+++ b/Engine/Source/Runtime/PacketHandlers/PacketHandler/Private/PacketHandlerProfileConfig.cpp
@@ -1,4 +1,4 @@
-// Copyright 1998-2018 Epic Games, Inc. All Rights Reserved.
+// Copyright 1998-2019 Epic Games, Inc. All Rights Reserved.
 
 #include "PacketHandlerProfileConfig.h"
 
