Fixing crashing while loading packethandler profile configurations due to a naming collision with transient packages.

diff --git a/Engine/Source/Runtime/PacketHandlers/PacketHandler/Private/PacketHandler.cpp b/Engine/Source/Runtime/PacketHandlers/PacketHandler/Private/PacketHandler.cpp
--- a/Engine/Source/Runtime/PacketHandlers/PacketHandler/Private/PacketHandler.cpp
+++ b/Engine/Source/Runtime/PacketHandlers/PacketHandler/Private/PacketHandler.cpp
@@ -97,8 +97,8 @@ void PacketHandler::Initialize(Handler::Mode InMode, uint32 InMaxPacketBits, boo
 	if (!bConnectionlessHandler)
 	{
 		TArray<FString> Components;
-		UPacketHandlerProfileConfig* CurNetDriverProfile = NewObject<UPacketHandlerProfileConfig>((UObject*)GetTransientPackage(), InDriverProfile);
-		Components.Append(CurNetDriverProfile->Components);
+		FString DriverProfileCategory = FString::Printf(TEXT("%s PacketHandlerProfileConfig"), *InDriverProfile.GetPlainNameString());
+		GConfig->GetArray(*DriverProfileCategory, TEXT("Components"), Components, GEngineIni);
 		
 		// If we didn't get any matches, push in the regular components.
 		if (Components.Num() == 0)
@@ -818,17 +818,8 @@ bool PacketHandler::DoesAnyProfileHaveComponent(const FString& InComponentName)
 bool PacketHandler::DoesProfileHaveComponent(const FName InNetDriverName, const FString& InComponentName)
 {
 	TArray<FString> Components;
-	// If this call was made during a module load, we cannot use the packethandler config object
-	if (GetTransientPackage() == nullptr)
-	{
-		FString DriverProfileCategory = FString::Printf(TEXT("%s PacketHandlerProfileConfig"), *InNetDriverName.ToString());
-		GConfig->GetArray(*DriverProfileCategory, TEXT("Components"), Components, GEngineIni);
-	}
-	else
-	{
-		UPacketHandlerProfileConfig* CurNetDriverProfile = NewObject<UPacketHandlerProfileConfig>((UObject*)GetTransientPackage(), InNetDriverName);
-		Components.Append(CurNetDriverProfile->Components);
-	}
+	FString DriverProfileCategory = FString::Printf(TEXT("%s PacketHandlerProfileConfig"), *InNetDriverName.GetPlainNameString());
+	GConfig->GetArray(*DriverProfileCategory, TEXT("Components"), Components, GEngineIni);
 	
 	for (const FString& Component : Components)
 	{
