Adjusted Net Analytiics Aggregation PerObjectConfig settings object, to not use the transient package, and to use 'OverridePerObjectConfigSection' to retain the current section names.

Converted a frequently hit ensure, to a warning log.

diff --git a/Engine/Source/Runtime/PacketHandlers/PacketHandler/Classes/NetAnalyticsAggregatorConfig.h b/Engine/Source/Runtime/PacketHandlers/PacketHandler/Classes/NetAnalyticsAggregatorConfig.h
--- a/Engine/Source/Runtime/PacketHandlers/PacketHandler/Classes/NetAnalyticsAggregatorConfig.h
+++ b/Engine/Source/Runtime/PacketHandlers/PacketHandler/Classes/NetAnalyticsAggregatorConfig.h
@@ -26,7 +26,6 @@ struct FNetAnalyticsDataConfig
 	bool bEnabled;
 };
 
-
 /**
  * Configuration for FNetAnalyticsAggregator - loaded PerObjectConfig, for each NetDriverName
  */
@@ -35,6 +34,9 @@ class UNetAnalyticsAggregatorConfig : public UObject
 {
 	GENERATED_UCLASS_BODY()
 
+protected:
+	void OverridePerObjectConfigSection(FString& SectionName) override;
+
 public:
 	/** Registers FNetAnalyticsData data holders, by DataName - and specifies whether they are enabled or disabled */
 	UPROPERTY(config)
diff --git a/Engine/Source/Runtime/PacketHandlers/PacketHandler/Private/NetAnalytics.cpp b/Engine/Source/Runtime/PacketHandlers/PacketHandler/Private/NetAnalytics.cpp
--- a/Engine/Source/Runtime/PacketHandlers/PacketHandler/Private/NetAnalytics.cpp
+++ b/Engine/Source/Runtime/PacketHandlers/PacketHandler/Private/NetAnalytics.cpp
@@ -69,12 +69,20 @@ void FNetAnalyticsAggregator::Init()
 
 void FNetAnalyticsAggregator::InitConfig()
 {
-	UNetAnalyticsAggregatorConfig* CurConfig = NewObject<UNetAnalyticsAggregatorConfig>((UObject*)GetTransientPackage(), NetDriverName);
+	UClass* ClassRef = UNetAnalyticsAggregatorConfig::StaticClass();
+	UNetAnalyticsAggregatorConfig* CurConfig = FindObject<UNetAnalyticsAggregatorConfig>(ClassRef, *NetDriverName.ToString());
+
+	if (CurConfig == nullptr)
+	{
+		CurConfig = NewObject<UNetAnalyticsAggregatorConfig>(ClassRef, NetDriverName);
+	}
 
 	AnalyticsDataConfigMap.Empty();
 
 	// If the config is hotfixed, make sure no data holders are currently active, as they can't be selectively hotfixed if loaded
-	ensure(AnalyticsDataMap.Num() == 0);
+	// (this does seem to happen, frequently - so limits the hotfixability of Net Analytics)
+	UE_CLOG(AnalyticsDataMap.Num() > 0, PacketHandlerLog, Warning,
+			TEXT("Net Analytics hotfixed while already active. Analytics hotfix changes may not be applied correctly."));
 
 	if (CurConfig != nullptr)
 	{
diff --git a/Engine/Source/Runtime/PacketHandlers/PacketHandler/Private/NetAnalyticsAggregatorConfig.cpp b/Engine/Source/Runtime/PacketHandlers/PacketHandler/Private/NetAnalyticsAggregatorConfig.cpp
--- a/Engine/Source/Runtime/PacketHandlers/PacketHandler/Private/NetAnalyticsAggregatorConfig.cpp
+++ b/Engine/Source/Runtime/PacketHandlers/PacketHandler/Private/NetAnalyticsAggregatorConfig.cpp
@@ -10,3 +10,8 @@ UNetAnalyticsAggregatorConfig::UNetAnalyticsAggregatorConfig(const FObjectInitia
 	: Super(ObjectInitializer)
 {
 }
+
+void UNetAnalyticsAggregatorConfig::OverridePerObjectConfigSection(FString& SectionName)
+{
+	SectionName = GetName() + TEXT(" ") + GetClass()->GetName();
+}
