Add memory tracking to PacketHandler related classes

diff --git a/Engine/Plugins/Online/OnlineSubsystemSteam/Source/Private/OnlineAuthHandlerSteam.cpp b/Engine/Plugins/Online/OnlineSubsystemSteam/Source/Private/OnlineAuthHandlerSteam.cpp
--- a/Engine/Plugins/Online/OnlineSubsystemSteam/Source/Private/OnlineAuthHandlerSteam.cpp
+++ b/Engine/Plugins/Online/OnlineSubsystemSteam/Source/Private/OnlineAuthHandlerSteam.cpp
@@ -123,6 +123,16 @@ FSteamAuthHandlerComponent::~FSteamAuthHandlerComponent()
 	}
 }
 
+void FSteamAuthHandlerComponent::CountBytes(FArchive& Ar) const
+{
+	HandlerComponent::CountBytes(Ar);
+
+	const SIZE_T SizeOfThis = sizeof(*this) - sizeof(HandlerComponent);
+	Ar.CountBytes(SizeOfThis, SizeOfThis);
+
+	UserTicket.CountBytes(Ar);
+}
+
 void FSteamAuthHandlerComponent::Initialize()
 {
 	if (!AuthInterface.IsValid() || !AuthInterface->IsSessionAuthEnabled())
diff --git a/Engine/Plugins/Online/OnlineSubsystemSteam/Source/Private/OnlineAuthHandlerSteam.h b/Engine/Plugins/Online/OnlineSubsystemSteam/Source/Private/OnlineAuthHandlerSteam.h
--- a/Engine/Plugins/Online/OnlineSubsystemSteam/Source/Private/OnlineAuthHandlerSteam.h
+++ b/Engine/Plugins/Online/OnlineSubsystemSteam/Source/Private/OnlineAuthHandlerSteam.h
@@ -14,6 +14,7 @@ class FSteamAuthHandlerComponent : public HandlerComponent
 public:
 	FSteamAuthHandlerComponent();
 	virtual ~FSteamAuthHandlerComponent();
+	virtual void CountBytes(FArchive& Ar) const override;
 	virtual void Initialize() override;
 	virtual void NotifyHandshakeBegin() override;
 
diff --git a/Engine/Plugins/Runtime/PacketHandlers/AESHandlerComponent/Source/Private/AESHandlerComponent.cpp b/Engine/Plugins/Runtime/PacketHandlers/AESHandlerComponent/Source/Private/AESHandlerComponent.cpp
--- a/Engine/Plugins/Runtime/PacketHandlers/AESHandlerComponent/Source/Private/AESHandlerComponent.cpp
+++ b/Engine/Plugins/Runtime/PacketHandlers/AESHandlerComponent/Source/Private/AESHandlerComponent.cpp
@@ -207,3 +207,24 @@ int32 FAESHandlerComponent::GetReservedPacketBits() const
 	// Worst case includes the encryption enabled bit, the termination bit, padding up to the next whole byte, and a block of padding.
 	return 2 + 7 + (BlockSizeInBytes * 8);
 }
+
+void FAESHandlerComponent::CountBytes(FArchive& Ar) const
+{
+	FEncryptionComponent::CountBytes(Ar);
+
+	const SIZE_T SizeOfThis = sizeof(*this) - sizeof(FEncryptionComponent);
+	Ar.CountBytes(SizeOfThis, SizeOfThis);
+
+	/*
+	Note, as of now, EncryptionContext is just typedef'd, but none of the base
+	types actually allocated memory directly in their classes (although there may be
+	global state).
+	if (FEncryptionContext const * const LocalContext = EncrpytionContext.Get())
+	{
+		LocalContext->CountBytes(Ar);
+	}
+	*/
+
+	Key.CountBytes(Ar);
+	Ciphertext.CountBytes(Ar);
+}
diff --git a/Engine/Plugins/Runtime/PacketHandlers/AESHandlerComponent/Source/Public/AESHandlerComponent.h b/Engine/Plugins/Runtime/PacketHandlers/AESHandlerComponent/Source/Public/AESHandlerComponent.h
--- a/Engine/Plugins/Runtime/PacketHandlers/AESHandlerComponent/Source/Public/AESHandlerComponent.h
+++ b/Engine/Plugins/Runtime/PacketHandlers/AESHandlerComponent/Source/Public/AESHandlerComponent.h
@@ -47,6 +47,7 @@ public:
 	virtual void IncomingConnectionless(const FString& Address, FBitReader& Packet) override;
 	virtual void OutgoingConnectionless(const FString& Address, FBitWriter& Packet, FOutPacketTraits& Traits) override;
 	virtual int32 GetReservedPacketBits() const override;
+	virtual void CountBytes(FArchive& Ar) const override;
 
 private:
 	TUniquePtr<FEncryptionContext> EncryptionContext;
diff --git a/Engine/Plugins/Runtime/PacketHandlers/CompressionComponents/Oodle/Source/OodleHandlerComponent/Private/OodleHandlerComponent.cpp b/Engine/Plugins/Runtime/PacketHandlers/CompressionComponents/Oodle/Source/OodleHandlerComponent/Private/OodleHandlerComponent.cpp
--- a/Engine/Plugins/Runtime/PacketHandlers/CompressionComponents/Oodle/Source/OodleHandlerComponent/Private/OodleHandlerComponent.cpp
+++ b/Engine/Plugins/Runtime/PacketHandlers/CompressionComponents/Oodle/Source/OodleHandlerComponent/Private/OodleHandlerComponent.cpp
@@ -304,6 +304,14 @@ OodleHandlerComponent::~OodleHandlerComponent()
 	FreeDictionary(ClientDictionary);
 }
 
+void OodleHandlerComponent::CountBytes(FArchive& Ar) const
+{
+	HandlerComponent::CountBytes(Ar);
+
+	const SIZE_T SizeOfThis = sizeof(*this) - sizeof(HandlerComponent);
+	Ar.CountBytes(SizeOfThis, SizeOfThis);
+}
+
 void OodleHandlerComponent::InitFirstRunConfig()
 {
 	// Check that the OodleHandlerComponent section exists, and if not, init with defaults
diff --git a/Engine/Plugins/Runtime/PacketHandlers/CompressionComponents/Oodle/Source/OodleHandlerComponent/Public/OodleHandlerComponent.h b/Engine/Plugins/Runtime/PacketHandlers/CompressionComponents/Oodle/Source/OodleHandlerComponent/Public/OodleHandlerComponent.h
--- a/Engine/Plugins/Runtime/PacketHandlers/CompressionComponents/Oodle/Source/OodleHandlerComponent/Public/OodleHandlerComponent.h
+++ b/Engine/Plugins/Runtime/PacketHandlers/CompressionComponents/Oodle/Source/OodleHandlerComponent/Public/OodleHandlerComponent.h
@@ -274,6 +274,8 @@ public:
 	/** Default Destructor */
 	~OodleHandlerComponent();
 
+	virtual void CountBytes(FArchive& Ar) const override;
+
 	/**
 	 * Initializes first-run config settings
 	 */
diff --git a/Engine/Source/Runtime/Engine/Private/PacketHandlers/StatelessConnectHandlerComponent.cpp b/Engine/Source/Runtime/Engine/Private/PacketHandlers/StatelessConnectHandlerComponent.cpp
--- a/Engine/Source/Runtime/Engine/Private/PacketHandlers/StatelessConnectHandlerComponent.cpp
+++ b/Engine/Source/Runtime/Engine/Private/PacketHandlers/StatelessConnectHandlerComponent.cpp
@@ -209,6 +209,20 @@ StatelessConnectHandlerComponent::StatelessConnectHandlerComponent()
 	bRequiresHandshake = true;
 }
 
+void StatelessConnectHandlerComponent::CountBytes(FArchive& Ar) const
+{
+	HandlerComponent::CountBytes(Ar);
+
+	const SIZE_T SizeOfThis = sizeof(*this) - sizeof(HandlerComponent);
+
+	for (int32 i = 0; i < SECRET_COUNT; ++i)
+	{
+		HandshakeSecret[i].CountBytes(Ar);
+	}
+
+	LastChallengeSuccessAddress.CountBytes(Ar);
+}
+
 void StatelessConnectHandlerComponent::NotifyHandshakeBegin()
 {
 	if (Handler->Mode == Handler::Mode::Client)
diff --git a/Engine/Source/Runtime/Engine/Public/PacketHandlers/StatelessConnectHandlerComponent.h b/Engine/Source/Runtime/Engine/Public/PacketHandlers/StatelessConnectHandlerComponent.h
--- a/Engine/Source/Runtime/Engine/Public/PacketHandlers/StatelessConnectHandlerComponent.h
+++ b/Engine/Source/Runtime/Engine/Public/PacketHandlers/StatelessConnectHandlerComponent.h
@@ -43,6 +43,8 @@ public:
 	 */
 	StatelessConnectHandlerComponent();
 
+	virtual void CountBytes(FArchive& Ar) const override;
+
 	virtual bool IsValid() const override { return true; }
 
 	virtual void NotifyHandshakeBegin() override;
diff --git a/Engine/Source/Runtime/PacketHandlers/EncryptionComponents/AsymmetricEncryption/RSAEncryptionHandlerComponent/Private/RSAEncryptionHandlerComponent.cpp b/Engine/Source/Runtime/PacketHandlers/EncryptionComponents/AsymmetricEncryption/RSAEncryptionHandlerComponent/Private/RSAEncryptionHandlerComponent.cpp
--- a/Engine/Source/Runtime/PacketHandlers/EncryptionComponents/AsymmetricEncryption/RSAEncryptionHandlerComponent/Private/RSAEncryptionHandlerComponent.cpp
+++ b/Engine/Source/Runtime/PacketHandlers/EncryptionComponents/AsymmetricEncryption/RSAEncryptionHandlerComponent/Private/RSAEncryptionHandlerComponent.cpp
@@ -44,6 +44,13 @@ RSAEncryptionHandlerComponent::RSAEncryptionHandlerComponent(int32 InKeySizeInBi
 	bRequiresHandshake = true;
 }
 
+void RSAEncryptionHandlerComponent::CountBytes(FArchive& Ar) const
+{
+	HandlerComponent::CountBytes(Ar);
+
+	const SIZE_T SizeOfThis = sizeof(*this) - sizeof(HandlerComponent);
+}
+
 void RSAEncryptionHandlerComponent::Initialize()
 {
 	Params.GenerateRandomWithKeySize(Rng, KeySizeInBits);
diff --git a/Engine/Source/Runtime/PacketHandlers/EncryptionComponents/AsymmetricEncryption/RSAEncryptionHandlerComponent/Public/RSAEncryptionHandlerComponent.h b/Engine/Source/Runtime/PacketHandlers/EncryptionComponents/AsymmetricEncryption/RSAEncryptionHandlerComponent/Public/RSAEncryptionHandlerComponent.h
--- a/Engine/Source/Runtime/PacketHandlers/EncryptionComponents/AsymmetricEncryption/RSAEncryptionHandlerComponent/Public/RSAEncryptionHandlerComponent.h
+++ b/Engine/Source/Runtime/PacketHandlers/EncryptionComponents/AsymmetricEncryption/RSAEncryptionHandlerComponent/Public/RSAEncryptionHandlerComponent.h
@@ -40,6 +40,8 @@ public:
 	*/
 	RSAEncryptionHandlerComponent(int32 InKeySizeInBits = 1024);
 
+	virtual void CountBytes(FArchive& Ar) const override;
+
 	/* Initializes the handler component */
 	virtual void Initialize() override;
 
diff --git a/Engine/Source/Runtime/PacketHandlers/EncryptionComponents/SymmetricEncryption/BlockEncryption/BlockEncryptionHandlerComponent/Private/BlockEncryptionHandlerComponent.cpp b/Engine/Source/Runtime/PacketHandlers/EncryptionComponents/SymmetricEncryption/BlockEncryption/BlockEncryptionHandlerComponent/Private/BlockEncryptionHandlerComponent.cpp
--- a/Engine/Source/Runtime/PacketHandlers/EncryptionComponents/SymmetricEncryption/BlockEncryption/BlockEncryptionHandlerComponent/Private/BlockEncryptionHandlerComponent.cpp
+++ b/Engine/Source/Runtime/PacketHandlers/EncryptionComponents/SymmetricEncryption/BlockEncryption/BlockEncryptionHandlerComponent/Private/BlockEncryptionHandlerComponent.cpp
@@ -28,6 +28,16 @@ BlockEncryptionHandlerComponent::BlockEncryptionHandlerComponent(BlockEncryptor*
 	KeySizeInBytes = InKeySizeInBytes != 0 ? InKeySizeInBytes : Encryptor->GetDefaultKeySize();
 }
 
+void BlockEncryptionHandlerComponent::CountBytes(FArchive& Ar) const
+{
+	HandlerComponent::CountBytes(Ar);
+
+	const SIZE_T SizeOfThis = sizeof(*this) - sizeof(HandlerComponent);
+	Ar.CountBytes(SizeOfThis, SizeOfThis);
+
+	Key.CountBytes(Ar);
+}
+
 void BlockEncryptionHandlerComponent::Initialize()
 {
 	if (Handler->Mode == Handler::Mode::Client)
diff --git a/Engine/Source/Runtime/PacketHandlers/EncryptionComponents/SymmetricEncryption/BlockEncryption/BlockEncryptionHandlerComponent/Public/BlockEncryptionHandlerComponent.h b/Engine/Source/Runtime/PacketHandlers/EncryptionComponents/SymmetricEncryption/BlockEncryption/BlockEncryptionHandlerComponent/Public/BlockEncryptionHandlerComponent.h
--- a/Engine/Source/Runtime/PacketHandlers/EncryptionComponents/SymmetricEncryption/BlockEncryption/BlockEncryptionHandlerComponent/Public/BlockEncryptionHandlerComponent.h
+++ b/Engine/Source/Runtime/PacketHandlers/EncryptionComponents/SymmetricEncryption/BlockEncryption/BlockEncryptionHandlerComponent/Public/BlockEncryptionHandlerComponent.h
@@ -51,6 +51,8 @@ public:
 	/* Initializes default data, can provide the key size in bytes */
 	BlockEncryptionHandlerComponent(BlockEncryptor* Encryptor = nullptr, uint32 KeySizeInBytes = 0);
 
+	virtual void CountBytes(FArchive& Ar) const override;
+
 	/* Initializes the handler component */
 	virtual void Initialize() override;
 
diff --git a/Engine/Source/Runtime/PacketHandlers/EncryptionComponents/SymmetricEncryption/StreamEncryption/StreamEncryptionHandlerComponent/Private/StreamEncryptionHandlerComponent.cpp b/Engine/Source/Runtime/PacketHandlers/EncryptionComponents/SymmetricEncryption/StreamEncryption/StreamEncryptionHandlerComponent/Private/StreamEncryptionHandlerComponent.cpp
--- a/Engine/Source/Runtime/PacketHandlers/EncryptionComponents/SymmetricEncryption/StreamEncryption/StreamEncryptionHandlerComponent/Private/StreamEncryptionHandlerComponent.cpp
+++ b/Engine/Source/Runtime/PacketHandlers/EncryptionComponents/SymmetricEncryption/StreamEncryption/StreamEncryptionHandlerComponent/Private/StreamEncryptionHandlerComponent.cpp
@@ -29,6 +29,15 @@ StreamEncryptionHandlerComponent::StreamEncryptionHandlerComponent(StreamEncrypt
 	KeySizeInBytes = InKeySizeInBytes != 0 ? InKeySizeInBytes : Encryptor->GetDefaultKeySize();
 }
 
+void StreamEncryptionHandlerComponent::CountBytes(FArchive& Ar) const
+{
+	HandlerComponent::CountBytes(Ar);
+
+	const SIZE_T SizeOfThis = sizeof(*this) - sizeof(HandlerComponent);
+	Ar.CountBytes(SizeOfThis, SizeOfThis);
+	Key.CountBytes(Ar);
+}
+
 void StreamEncryptionHandlerComponent::Initialize()
 {
 	if (Handler->Mode == Handler::Mode::Client)
diff --git a/Engine/Source/Runtime/PacketHandlers/EncryptionComponents/SymmetricEncryption/StreamEncryption/StreamEncryptionHandlerComponent/Public/StreamEncryptionHandlerComponent.h b/Engine/Source/Runtime/PacketHandlers/EncryptionComponents/SymmetricEncryption/StreamEncryption/StreamEncryptionHandlerComponent/Public/StreamEncryptionHandlerComponent.h
--- a/Engine/Source/Runtime/PacketHandlers/EncryptionComponents/SymmetricEncryption/StreamEncryption/StreamEncryptionHandlerComponent/Public/StreamEncryptionHandlerComponent.h
+++ b/Engine/Source/Runtime/PacketHandlers/EncryptionComponents/SymmetricEncryption/StreamEncryption/StreamEncryptionHandlerComponent/Public/StreamEncryptionHandlerComponent.h
@@ -40,6 +40,8 @@ public:
 	/* Initializes default data, can provide the key size in bytes */
 	StreamEncryptionHandlerComponent(StreamEncryptor* Encryptor = nullptr, uint32 KeySizeInBytes = 0);
 
+	virtual void CountBytes(FArchive& Ar) const override;
+
 	/* Initializes the handler component */
 	virtual void Initialize() override;
 
diff --git a/Engine/Source/Runtime/PacketHandlers/PacketHandler/Private/PacketHandler.cpp b/Engine/Source/Runtime/PacketHandlers/PacketHandler/Private/PacketHandler.cpp
--- a/Engine/Source/Runtime/PacketHandlers/PacketHandler/Private/PacketHandler.cpp
+++ b/Engine/Source/Runtime/PacketHandlers/PacketHandler/Private/PacketHandler.cpp
@@ -379,6 +379,54 @@ TSharedPtr<HandlerComponent> PacketHandler::GetComponentByName(FName ComponentNa
 	return nullptr;
 }
 
+void PacketHandler::CountBytes(FArchive& Ar) const
+{
+	Ar.CountBytes(sizeof(*this), sizeof(*this));
+	OutgoingPacket.CountMemory(Ar);
+	IncomingPacket.CountMemory(Ar);
+
+	HandlerComponents.CountBytes(Ar);
+	for (const TSharedPtr<HandlerComponent>& Component : HandlerComponents)
+	{
+		if (HandlerComponent const * const LocalComponent = Component.Get())
+		{
+			LocalComponent->CountBytes(Ar);
+		}
+	}
+
+	// Don't handle EncryptionComponent, as it should be in our components array.
+
+	BufferedPackets.CountBytes(Ar);
+	for (BufferedPacket const * const LocalPacket : BufferedPackets)
+	{
+		if (LocalPacket)
+		{
+			LocalPacket->CountBytes(Ar);
+		}
+	}
+
+	// Unfortunately, there's currently no way to safely calculate memory usage for TQueues.
+	// so QueuedPackets, QueuedRawPackets, QueuedHandlerPackets, and QueuedConnectionlessPackets
+	// can't be tracked without a rework.
+
+	BufferedConnectionlessPackets.CountBytes(Ar);
+	for (BufferedPacket const * const LocalPacket : BufferedConnectionlessPackets)
+	{
+		if (LocalPacket)
+		{
+			LocalPacket->CountBytes(Ar);
+		}
+	}
+
+	// Don't handle ReliabilityComponent, since it should be in our components array.
+	// Don't track AnalyticsProvider as that should be handled elsewhere.
+}
+
+void HandlerComponent::CountBytes(FArchive& Ar) const
+{
+	Ar.CountBytes(sizeof(*this), sizeof(*this));
+}
+
 const ProcessedPacket PacketHandler::Incoming_Internal(uint8* Packet, int32 CountBytes, bool bConnectionless, const FString& Address)
 {
 	SCOPE_CYCLE_COUNTER(Stat_PacketHandler_Incoming_Internal);
diff --git a/Engine/Source/Runtime/PacketHandlers/PacketHandler/Public/PacketHandler.h b/Engine/Source/Runtime/PacketHandlers/PacketHandler/Public/PacketHandler.h
--- a/Engine/Source/Runtime/PacketHandlers/PacketHandler/Public/PacketHandler.h
+++ b/Engine/Source/Runtime/PacketHandlers/PacketHandler/Public/PacketHandler.h
@@ -210,6 +210,13 @@ public:
 	{
 		delete [] Data;
 	}
+
+	void CountBytes(FArchive& Ar) const
+	{
+		Ar.CountBytes(sizeof(*this), sizeof(*this));
+		Ar.CountBytes(FMath::DivideAndRoundUp(CountBits, 8u), FMath::DivideAndRoundUp(CountBits, 8u));
+		Address.CountBytes(Ar);
+	}
 };
 
 /**
@@ -406,6 +413,8 @@ public:
 	/** Returns a pointer to the first component in the HandlerComponents array with the specified name. */
 	TSharedPtr<HandlerComponent> GetComponentByName(FName ComponentName) const;
 
+	virtual void CountBytes(FArchive& Ar) const;
+
 protected:
 	/**
 	 * Internal handling for Incoming/IncomingConnectionless
@@ -829,6 +838,8 @@ public:
 	 */
 	virtual void NotifyAnalyticsProvider() {}
 
+	virtual void CountBytes(FArchive& Ar) const;
+
 protected:
 	/**
 	 * Sets the state of the handler
diff --git a/Engine/Source/Runtime/PacketHandlers/ReliabilityHandlerComponent/Private/ReliabilityHandlerComponent.cpp b/Engine/Source/Runtime/PacketHandlers/ReliabilityHandlerComponent/Private/ReliabilityHandlerComponent.cpp
--- a/Engine/Source/Runtime/PacketHandlers/ReliabilityHandlerComponent/Private/ReliabilityHandlerComponent.cpp
+++ b/Engine/Source/Runtime/PacketHandlers/ReliabilityHandlerComponent/Private/ReliabilityHandlerComponent.cpp
@@ -19,6 +19,23 @@ ReliabilityHandlerComponent::ReliabilityHandlerComponent()
 {
 }
 
+void ReliabilityHandlerComponent::CountBytes(FArchive& Ar) const
+{
+	HandlerComponent::CountBytes(Ar);
+
+	const SIZE_T SizeOfThis = sizeof(*this) - sizeof(HandlerComponent);
+	Ar.CountBytes(SizeOfThis, SizeOfThis);
+
+	BufferedPackets.CountBytes(Ar);
+	for (BufferedPacket const * const LocalPacket : BufferedPackets)
+	{
+		if (LocalPacket)
+		{
+			LocalPacket->CountBytes(Ar);
+		}
+	}
+}
+
 void ReliabilityHandlerComponent::Initialize()
 {
 	SetActive(true);
diff --git a/Engine/Source/Runtime/PacketHandlers/ReliabilityHandlerComponent/Public/ReliabilityHandlerComponent.h b/Engine/Source/Runtime/PacketHandlers/ReliabilityHandlerComponent/Public/ReliabilityHandlerComponent.h
--- a/Engine/Source/Runtime/PacketHandlers/ReliabilityHandlerComponent/Public/ReliabilityHandlerComponent.h
+++ b/Engine/Source/Runtime/PacketHandlers/ReliabilityHandlerComponent/Public/ReliabilityHandlerComponent.h
@@ -13,6 +13,8 @@ public:
 	/* Initializes default data */
 	ReliabilityHandlerComponent();
 
+	virtual void CountBytes(FArchive& Ar) const;
+
 	virtual void Initialize() override;
 
 	virtual bool IsValid() const override;
