Updating NetConnection, NetDriver and PacketHandlers to use FInternetAddrs. Moves address information to the base class instead of relying on it being created by derived classes.

diff --git a/Engine/Plugins/Experimental/HTML5Networking/Source/HTML5Networking/Classes/WebSocketConnection.h b/Engine/Plugins/Experimental/HTML5Networking/Source/HTML5Networking/Classes/WebSocketConnection.h
--- a/Engine/Plugins/Experimental/HTML5Networking/Source/HTML5Networking/Classes/WebSocketConnection.h
+++ b/Engine/Plugins/Experimental/HTML5Networking/Source/HTML5Networking/Classes/WebSocketConnection.h
@@ -20,10 +20,6 @@ class HTML5NETWORKING_API UWebSocketConnection : public UNetConnection
 	virtual void LowLevelSend(void* Data, int32 CountBits, FOutPacketTraits& Traits) override;
 	FString LowLevelGetRemoteAddress(bool bAppendPort = false) override;
 	FString LowLevelDescribe() override;
-	virtual int32 GetAddrAsInt(void) override;
-	virtual int32 GetAddrPort(void) override;
-	virtual TSharedPtr<FInternetAddr> GetInternetAddr() override;
-	virtual FString RemoteAddressToString() override;
 	virtual void Tick();
 	virtual void FinishDestroy();
 	virtual void ReceivedRawPacket(void* Data,int32 Count);
diff --git a/Engine/Plugins/Experimental/HTML5Networking/Source/HTML5Networking/Classes/WebSocketNetDriver.h b/Engine/Plugins/Experimental/HTML5Networking/Source/HTML5Networking/Classes/WebSocketNetDriver.h
--- a/Engine/Plugins/Experimental/HTML5Networking/Source/HTML5Networking/Classes/WebSocketNetDriver.h
+++ b/Engine/Plugins/Experimental/HTML5Networking/Source/HTML5Networking/Classes/WebSocketNetDriver.h
@@ -27,7 +27,7 @@ class HTML5NETWORKING_API UWebSocketNetDriver : public UNetDriver
 	virtual bool InitConnect(FNetworkNotify* InNotify, const FURL& ConnectURL, FString& Error) override;
 	virtual bool InitListen(FNetworkNotify* InNotify, FURL& LocalURL, bool bReuseAddressAndPort, FString& Error) override;
 	virtual void TickDispatch(float DeltaTime) override;
-	virtual void LowLevelSend(FString Address, void* Data, int32 CountBits, FOutPacketTraits& Traits) override;
+	virtual void LowLevelSend(TSharedPtr<const FInternetAddr> Address, void* Data, int32 CountBits, FOutPacketTraits& Traits) override;
 	virtual FString LowLevelGetNetworkNumber() override;
 	virtual void LowLevelDestroy() override;
 	virtual bool IsNetResourceValid(void) override;
diff --git a/Engine/Plugins/Experimental/HTML5Networking/Source/HTML5Networking/Private/WebSocket.cpp b/Engine/Plugins/Experimental/HTML5Networking/Source/HTML5Networking/Private/WebSocket.cpp
--- a/Engine/Plugins/Experimental/HTML5Networking/Source/HTML5Networking/Private/WebSocket.cpp
+++ b/Engine/Plugins/Experimental/HTML5Networking/Source/HTML5Networking/Private/WebSocket.cpp
@@ -90,7 +90,7 @@ FWebSocket::FWebSocket(
 	SockFd = socket(PF_INET, SOCK_STREAM, IPPROTO_TCP);
 	if (SockFd == -1)
 	{
-		UE_LOG(LogHTML5Networking, Error, TEXT("Socket creationg failed "));
+		UE_LOG(LogHTML5Networking, Error, TEXT("Socket creation failed "));
 	}
 	else
 	{
@@ -101,9 +101,10 @@ FWebSocket::FWebSocket(
 
 #endif
 
+	memset(&RemoteAddr, 0, sizeof(RemoteAddr));
+
 	// Windows XP does not have support for inet_pton
 #if PLATFORM_WINDOWS && _WIN32_WINNT <= 0x0502
-	memset(&RemoteAddr, 0, sizeof(RemoteAddr));
 	int32 SizeOfRemoteAddr = sizeof(RemoteAddr);
 
 	// Force ServerAddress into non-const array. API doesn't modify contents but old API still requires non-const string
@@ -112,24 +113,20 @@ FWebSocket::FWebSocket(
 		UE_LOG(LogHTML5Networking, Warning, TEXT("WSAStringToAddress failed "));
 		return;
 	}
-
-	RemoteAddr.sin_family = AF_INET;
-	RemoteAddr.sin_port = htons(ServerAddress.GetPort());
 #else
-	memset(&RemoteAddr, 0, sizeof(RemoteAddr));
-	RemoteAddr.sin_family = AF_INET;
-	RemoteAddr.sin_port = htons(ServerAddress.GetPort());
-
 	if (inet_pton(AF_INET, TCHAR_TO_ANSI(*ServerAddress.ToString(false)), &RemoteAddr.sin_addr) != 1)
 	{
-		UE_LOG(LogHTML5Networking, Warning, TEXT("inet_pton failed "));
+		UE_LOG(LogHTML5Networking, Warning, TEXT("inet_pton failed to %s"), *ServerAddress.ToString(false));
 		return;
 	}
 #endif
 
+	RemoteAddr.sin_family = AF_INET;
+	RemoteAddr.sin_port = htons(ServerAddress.GetPort());
+
 #if !USE_LIBWEBSOCKET // HTML5 uses BSD network API
 	int Ret = connect(SockFd, (struct sockaddr *)&RemoteAddr, sizeof(RemoteAddr));
-	UE_LOG(LogHTML5Networking, Warning, TEXT(" Connect socket returned %d"), Ret);
+	UE_LOG(LogHTML5Networking, Warning, TEXT(" Connect socket returned %d to %s. Error Code: %d"), Ret, *ServerAddress.ToString(false), ((Ret != 0) ? errno : 0));
 #endif
 }
 
@@ -194,11 +191,6 @@ FString FWebSocket::RemoteEndPoint(bool bAppendPort)
 	return remote;
 }
 
-struct sockaddr_in* FWebSocket::GetRemoteAddr()
-{
-	return &RemoteAddr;
-}
-
 FString FWebSocket::LocalEndPoint(bool bAppendPort)
 {
 #if USE_LIBWEBSOCKET
@@ -305,6 +297,19 @@ void FWebSocket::Flush()
 	};
 }
 
+TArray<uint8> FWebSocket::GetRawRemoteAddr(int32& OutPort)
+{
+	OutPort = ntohs(RemoteAddr.sin_port);
+	TArray<uint8> RawBuffer;
+	uint32 IntAddr = RemoteAddr.sin_addr.s_addr;
+	RawBuffer.Add((IntAddr >> 0) & 0xFF);
+	RawBuffer.Add((IntAddr >> 8) & 0xFF);
+	RawBuffer.Add((IntAddr >> 16) & 0xFF);
+	RawBuffer.Add((IntAddr >> 24) & 0xFF);
+
+	return RawBuffer;
+}
+
 void FWebSocket::SetConnectedCallBack(FWebsocketInfoCallBack CallBack)
 {
 	ConnectedCallBack = CallBack;
diff --git a/Engine/Plugins/Experimental/HTML5Networking/Source/HTML5Networking/Private/WebSocket.h b/Engine/Plugins/Experimental/HTML5Networking/Source/HTML5Networking/Private/WebSocket.h
--- a/Engine/Plugins/Experimental/HTML5Networking/Source/HTML5Networking/Private/WebSocket.h
+++ b/Engine/Plugins/Experimental/HTML5Networking/Source/HTML5Networking/Private/WebSocket.h
@@ -44,9 +44,10 @@ public:
 	void Flush();
 
 	/** Helper functions to describe end points. */
+	TArray<uint8> GetRawRemoteAddr(int32& OutPort);
 	FString RemoteEndPoint(bool bAppendPort);
 	FString LocalEndPoint(bool bAppendPort);
-	struct sockaddr_in* GetRemoteAddr();
+	struct sockaddr_in* GetRemoteAddr() { return &RemoteAddr; }
 
 // this was made public because of cross-platform build issues
 public:
diff --git a/Engine/Plugins/Experimental/HTML5Networking/Source/HTML5Networking/Private/WebSocketNetDriver.cpp b/Engine/Plugins/Experimental/HTML5Networking/Source/HTML5Networking/Private/WebSocketNetDriver.cpp
--- a/Engine/Plugins/Experimental/HTML5Networking/Source/HTML5Networking/Private/WebSocketNetDriver.cpp
+++ b/Engine/Plugins/Experimental/HTML5Networking/Source/HTML5Networking/Private/WebSocketNetDriver.cpp
@@ -106,6 +106,8 @@ bool UWebSocketNetDriver::InitListen(FNetworkNotify* InNotify, FURL& LocalURL, b
 
 	WebSocketServer->Tick();
 	LocalURL.Port = WebSocketPort;
+	LocalAddr = GetSocketSubsystem()->GetLocalBindAddr(*GLog);
+	LocalAddr->SetPort(WebSocketPort);
 	UE_LOG(LogHTML5Networking, Log, TEXT("%s WebSocketNetDriver listening on port %i"), *GetDescription(), LocalURL.Port);
 
 	// server has no server connection.
@@ -121,17 +123,9 @@ void UWebSocketNetDriver::TickDispatch(float DeltaTime)
 		WebSocketServer->Tick();
 }
 
-void UWebSocketNetDriver::LowLevelSend(FString Address, void* Data, int32 CountBits, FOutPacketTraits& Traits)
+void UWebSocketNetDriver::LowLevelSend(TSharedPtr<const FInternetAddr> Address, void* Data, int32 CountBits, FOutPacketTraits& Traits)
 {
-	bool bValidAddress = !Address.IsEmpty();
-	TSharedRef<FInternetAddr> RemoteAddr = GetSocketSubsystem()->CreateInternetAddr();
-
-	if (bValidAddress)
-	{
-		RemoteAddr->SetIp(*Address, bValidAddress);
-	}
-
-	if (bValidAddress)
+	if (Address.IsValid() && Address->IsValid())
 	{
 		const uint8* DataToSend = reinterpret_cast<uint8*>(Data);
 
@@ -159,7 +153,7 @@ void UWebSocketNetDriver::LowLevelSend(FString Address, void* Data, int32 CountB
 			for (int32 i = 0; i<ClientConnections.Num(); ++i)
 			{
 				UWebSocketConnection* Connection = (UWebSocketConnection*)ClientConnections[i];
-				if (Connection && ( Connection->LowLevelGetRemoteAddress(true) == Address ) )
+				if (Connection && ( Connection->LowLevelGetRemoteAddress(true) == Address->ToString(true) ) )
 				{
 					Connection->GetWebSocket()->Send((uint8*)DataToSend, FMath::DivideAndRoundUp(CountBits, 8));
 					break;
@@ -169,13 +163,13 @@ void UWebSocketNetDriver::LowLevelSend(FString Address, void* Data, int32 CountB
 	}
 	else
 	{
-		UE_LOG(LogNet, Warning, TEXT("UWebSocketNetDriver::LowLevelSend: Invalid send address '%s'"), *Address);
+		UE_LOG(LogNet, Warning, TEXT("UWebSocketNetDriver::LowLevelSend: Invalid send address '%s'"), *Address->ToString(false));
 	}
 }
 
 FString UWebSocketNetDriver::LowLevelGetNetworkNumber()
 {
-	return WebSocketServer->Info();
+	return WebSocketServer != nullptr ? WebSocketServer->Info() : FString(TEXT(""));
 }
 
 void UWebSocketNetDriver::LowLevelDestroy()
@@ -225,9 +219,9 @@ void UWebSocketNetDriver::OnWebSocketClientConnected(FWebSocket* ClientWebSocket
 		check(Connection);
 
 		TSharedRef<FInternetAddr> InternetAddr = GetSocketSubsystem()->CreateInternetAddr();
-		bool Ok;
+		int32 AddressPort = 0;
 
-		InternetAddr->SetIp(*(ClientWebSocket->RemoteEndPoint(false)),Ok);
+		InternetAddr->SetRawIp(ClientWebSocket->GetRawRemoteAddr(AddressPort));
 		InternetAddr->SetPort(0);
 		Connection->SetWebSocket(ClientWebSocket);
 		Connection->InitRemoteConnection(this, NULL, FURL(), *InternetAddr, USOCK_Open);
diff --git a/Engine/Plugins/Experimental/HTML5Networking/Source/HTML5Networking/Private/WebsocketConnection.cpp b/Engine/Plugins/Experimental/HTML5Networking/Source/HTML5Networking/Private/WebsocketConnection.cpp
--- a/Engine/Plugins/Experimental/HTML5Networking/Source/HTML5Networking/Private/WebsocketConnection.cpp
+++ b/Engine/Plugins/Experimental/HTML5Networking/Source/HTML5Networking/Private/WebsocketConnection.cpp
@@ -166,19 +166,24 @@ void UWebSocketConnection::ReceivedRawPacket(void* Data,int32 Count)
 	if ( bChallengeHandshake )
 	{
 		// Process all incoming packets.
-		if (Driver->ConnectionlessHandler.IsValid() && Driver->StatelessConnectComponent.IsValid())
+		if (Driver->ConnectionlessHandler.IsValid() && Driver->StatelessConnectComponent.IsValid() && Driver->GetSocketSubsystem() != nullptr)
 		{
+			int32 Port;
+			TSharedPtr<FInternetAddr> IncomingAddress = Driver->GetSocketSubsystem()->CreateInternetAddr();
+			IncomingAddress->SetRawIp(WebSocket->GetRawRemoteAddr(Port));
+			IncomingAddress->SetPort(Port);
+
 			const ProcessedPacket UnProcessedPacket =
-									Driver->ConnectionlessHandler->IncomingConnectionless(LowLevelGetRemoteAddress(true), DataRef, Count);
+									Driver->ConnectionlessHandler->IncomingConnectionless(IncomingAddress, DataRef, Count);
 	
 			TSharedPtr<StatelessConnectHandlerComponent> StatelessConnect = Driver->StatelessConnectComponent.Pin();
 
 			bool bRestartedHandshake = false;
 
-			if (!UnProcessedPacket.bError && StatelessConnect->HasPassedChallenge(LowLevelGetRemoteAddress(true), bRestartedHandshake) &&
+			if (!UnProcessedPacket.bError && StatelessConnect->HasPassedChallenge(IncomingAddress, bRestartedHandshake) &&
 					!bRestartedHandshake)
 			{
-				UE_LOG(LogNet, Log, TEXT("Server accepting post-challenge connection from: %s"), *LowLevelGetRemoteAddress(true));
+				UE_LOG(LogNet, Log, TEXT("Server accepting post-challenge connection from: %s"), *(IncomingAddress->ToString(false)));
 				// Set the initial packet sequence from the handshake data
 				if (StatelessConnectComponent.IsValid())
 				{
@@ -215,30 +220,3 @@ void UWebSocketConnection::ReceivedRawPacket(void* Data,int32 Count)
 
 	UNetConnection::ReceivedRawPacket(DataRef,Count);
 }
-
-int32 UWebSocketConnection::GetAddrAsInt()
-{
-	// Get the host byte order ip addr
-	struct sockaddr_in* sock = WebSocket->GetRemoteAddr();
-	return (int32)ntohl(sock->sin_addr.s_addr);
-}
-
-int32 UWebSocketConnection::GetAddrPort()
-{
-	// Get the host byte order ip port
-	struct sockaddr_in* sock = WebSocket->GetRemoteAddr();
-	return (int32)ntohs(sock->sin_port);
-}
-
-TSharedPtr<FInternetAddr> UWebSocketConnection::GetInternetAddr()
-{
-	struct sockaddr_in* sock = WebSocket->GetRemoteAddr();
-
-	// @todo #JIRA UENET-883: This should be based on NetConnection.RemoteAddr, when moved down from IPConnection
-	return ISocketSubsystem::Get()->CreateInternetAddr((int32)ntohl(sock->sin_addr.s_addr), (int32)ntohs(sock->sin_port));
-}
-
-FString UWebSocketConnection::RemoteAddressToString()
-{
-	return WebSocket->RemoteEndPoint(true);
-}
diff --git a/Engine/Plugins/Online/OnlineSubsystemOculus/Source/Classes/OculusNetConnection.h b/Engine/Plugins/Online/OnlineSubsystemOculus/Source/Classes/OculusNetConnection.h
--- a/Engine/Plugins/Online/OnlineSubsystemOculus/Source/Classes/OculusNetConnection.h
+++ b/Engine/Plugins/Online/OnlineSubsystemOculus/Source/Classes/OculusNetConnection.h
@@ -30,7 +30,6 @@ public:
 	FString LowLevelDescribe() override;
 	virtual void FinishDestroy() override;
 
-	virtual TSharedPtr<FInternetAddr> GetInternetAddr() override;
 	virtual FString RemoteAddressToString() override;
 	// End NetConnection Interface
 	
diff --git a/Engine/Plugins/Online/OnlineSubsystemOculus/Source/Classes/OculusNetDriver.h b/Engine/Plugins/Online/OnlineSubsystemOculus/Source/Classes/OculusNetDriver.h
--- a/Engine/Plugins/Online/OnlineSubsystemOculus/Source/Classes/OculusNetDriver.h
+++ b/Engine/Plugins/Online/OnlineSubsystemOculus/Source/Classes/OculusNetDriver.h
@@ -5,6 +5,7 @@
 #include "CoreMinimal.h"
 #include "OculusNetConnection.h"
 #include "IpNetDriver.h"
+#include "IPAddress.h"
 #include "OculusNetDriver.generated.h"
 
 /**
@@ -35,7 +36,7 @@ public:
 	virtual bool InitConnect(FNetworkNotify* InNotify, const FURL& ConnectURL, FString& Error) override;
 	virtual bool InitListen(FNetworkNotify* InNotify, FURL& LocalURL, bool bReuseAddressAndPort, FString& Error) override;
 	virtual void TickDispatch(float DeltaTime) override;
-	virtual void LowLevelSend(FString Address, void* Data, int32 CountBits, FOutPacketTraits& Traits) override;
+	virtual void LowLevelSend(TSharedPtr<const FInternetAddr> Address, void* Data, int32 CountBits, FOutPacketTraits& Traits) override;
 	virtual void Shutdown() override;
 	virtual bool IsNetResourceValid() override;
 
diff --git a/Engine/Plugins/Online/OnlineSubsystemOculus/Source/Private/OculusNetConnection.cpp b/Engine/Plugins/Online/OnlineSubsystemOculus/Source/Private/OculusNetConnection.cpp
--- a/Engine/Plugins/Online/OnlineSubsystemOculus/Source/Private/OculusNetConnection.cpp
+++ b/Engine/Plugins/Online/OnlineSubsystemOculus/Source/Private/OculusNetConnection.cpp
@@ -62,8 +62,8 @@ void UOculusNetConnection::InitRemoteConnection(UNetDriver* InDriver, class FSoc
 		InMaxPacket == 0 ? MAX_PACKET_SIZE : InMaxPacket,
 		0);
 
-	auto OculusAddr = static_cast<const FInternetAddrOculus&>(InRemoteAddr);
-	PeerID = OculusAddr.GetID();
+	RemoteAddr = InRemoteAddr.Clone();
+	PeerID = StaticCastSharedPtr<FInternetAddrOculus>(RemoteAddr)->GetID();
 
 	// This is for a client that needs to log in, setup ClientLoginState and ExpectedClientLoginMsgType to reflect that
 	SetClientLoginState(EClientLoginState::LoggingIn);
@@ -156,12 +156,6 @@ void UOculusNetConnection::FinishDestroy()
 	}
 }
 
-TSharedPtr<FInternetAddr> UOculusNetConnection::GetInternetAddr()
-{
-	// @todo #JIRA UENET-883: This should be based on NetConnection.RemoteAddr, when moved down from IPConnection
-	return MakeShareable(new FInternetAddrOculus(PeerID));
-}
-
 FString UOculusNetConnection::RemoteAddressToString()
 {
 	if (bIsPassThrough)
diff --git a/Engine/Plugins/Online/OnlineSubsystemOculus/Source/Private/OculusNetDriver.cpp b/Engine/Plugins/Online/OnlineSubsystemOculus/Source/Private/OculusNetDriver.cpp
--- a/Engine/Plugins/Online/OnlineSubsystemOculus/Source/Private/OculusNetDriver.cpp
+++ b/Engine/Plugins/Online/OnlineSubsystemOculus/Source/Private/OculusNetDriver.cpp
@@ -199,14 +199,14 @@ void UOculusNetDriver::TickDispatch(float DeltaTime)
 			}
 
 			UE_LOG(LogNet, Verbose, TEXT("Checking challenge from: %llu"), PeerID);
+			TSharedPtr<FInternetAddr> OculusAddr = MakeShareable(new FInternetAddrOculus(PeerID));
 			StatelessConnect = StatelessConnectComponent.Pin();
-			FString IncomingAddress = FString::Printf(TEXT("%llu.oculus"), PeerID);
 
 			const ProcessedPacket UnProcessedPacket =
-				ConnectionlessHandler->IncomingConnectionless(IncomingAddress, Data, PacketSize);
+				ConnectionlessHandler->IncomingConnectionless(OculusAddr, Data, PacketSize);
 
 			bool bRestartedHandshake = false;
-			bPassedChallenge = !UnProcessedPacket.bError && StatelessConnect->HasPassedChallenge(IncomingAddress, bRestartedHandshake);
+			bPassedChallenge = !UnProcessedPacket.bError && StatelessConnect->HasPassedChallenge(OculusAddr, bRestartedHandshake);
 
 			if (bPassedChallenge)
 			{
@@ -224,9 +224,7 @@ void UOculusNetDriver::TickDispatch(float DeltaTime)
 				UOculusNetConnection* Connection = NewObject<UOculusNetConnection>(NetConnectionClass);
 				check(Connection);
 
-				auto OculusAddr = new FInternetAddrOculus(PeerID);
-				TSharedRef<FInternetAddr> InternetAddr = MakeShareable(OculusAddr);
-				Connection->InitRemoteConnection(this, nullptr, FURL(), *InternetAddr, ovr_Net_IsConnected(PeerID) ? USOCK_Open : USOCK_Pending);
+				Connection->InitRemoteConnection(this, nullptr, FURL(), *OculusAddr, ovr_Net_IsConnected(PeerID) ? USOCK_Open : USOCK_Pending);
 
 				AddClientConnection(Connection);
 
@@ -287,13 +285,13 @@ void UOculusNetDriver::TickDispatch(float DeltaTime)
 	}
 }
 
-void UOculusNetDriver::LowLevelSend(FString Address, void* Data, int32 CountBits, FOutPacketTraits& Traits)
+void UOculusNetDriver::LowLevelSend(TSharedPtr<const FInternetAddr> Address, void* Data, int32 CountBits, FOutPacketTraits& Traits)
 {
 	if (bIsPassthrough)
 	{
 		return UIpNetDriver::LowLevelSend(Address, Data, CountBits, Traits);
 	}
-	FInternetAddrOculus OculusAddr(FURL(nullptr, *Address, ETravelType::TRAVEL_Absolute));
+	FInternetAddrOculus OculusAddr(FURL(nullptr, *Address->ToString(false), ETravelType::TRAVEL_Absolute));
 	ovrID PeerID = OculusAddr.GetID();
 	if (ovr_Net_IsConnected(PeerID))
 	{
diff --git a/Engine/Plugins/Online/OnlineSubsystemSteam/Source/Private/OnlineAuthHandlerSteam.h b/Engine/Plugins/Online/OnlineSubsystemSteam/Source/Private/OnlineAuthHandlerSteam.h
--- a/Engine/Plugins/Online/OnlineSubsystemSteam/Source/Private/OnlineAuthHandlerSteam.h
+++ b/Engine/Plugins/Online/OnlineSubsystemSteam/Source/Private/OnlineAuthHandlerSteam.h
@@ -7,6 +7,7 @@
 #include "OnlineSubsystemSteamTypes.h"
 #include "PacketHandler.h"
 #include "HandlerComponentFactory.h"
+#include "IPAddress.h"
 #include "OnlineAuthHandlerSteam.generated.h"
 
 class FSteamAuthHandlerComponent : public HandlerComponent
@@ -23,8 +24,8 @@ public:
 	virtual void Incoming(FBitReader& Packet) override;
 	virtual void Outgoing(FBitWriter& Packet, FOutPacketTraits& Traits) override;
 
-	virtual void IncomingConnectionless(const FString& Address, FBitReader& Packet) override {}
-	virtual void OutgoingConnectionless(const FString& Address, FBitWriter& Packet, FOutPacketTraits& Traits) override {}
+	virtual void IncomingConnectionless(const TSharedPtr<const FInternetAddr>& Address, FBitReader& Packet) override {}
+	virtual void OutgoingConnectionless(const TSharedPtr<const FInternetAddr>& Address, FBitWriter& Packet, FOutPacketTraits& Traits) override {}
 
 	virtual void Tick(float DeltaTime) override;
 
diff --git a/Engine/Plugins/Online/OnlineSubsystemSteam/Source/Private/SteamNetDriver.cpp b/Engine/Plugins/Online/OnlineSubsystemSteam/Source/Private/SteamNetDriver.cpp
--- a/Engine/Plugins/Online/OnlineSubsystemSteam/Source/Private/SteamNetDriver.cpp
+++ b/Engine/Plugins/Online/OnlineSubsystemSteam/Source/Private/SteamNetDriver.cpp
@@ -2,8 +2,8 @@
 
 #include "SteamNetDriver.h"
 #include "OnlineSubsystemNames.h"
-#include "OnlineSubsystem.h"
 #include "SocketSubsystem.h"
+#include "OnlineSubsystemSteam.h"
 #include "OnlineSubsystemSteamPrivate.h"
 #include "SocketsSteam.h"
 #include "SteamNetConnection.h"
@@ -89,7 +89,7 @@ bool USteamNetDriver::InitConnect(FNetworkNotify* InNotify, const FURL& ConnectU
 		// If we are opening a Steam URL, create a Steam client socket
 		if (ConnectURL.Host.StartsWith(STEAM_URL_PREFIX))
 		{
-			Socket = SteamSockets->CreateSocket(FName(TEXT("SteamClientSocket")), TEXT("Unreal client (Steam)"));
+			Socket = SteamSockets->CreateSocket(FName(TEXT("SteamClientSocket")), TEXT("Unreal client (Steam)"), FNetworkProtocolTypes::Steam);
 		}
 		else
 		{
@@ -106,7 +106,7 @@ bool USteamNetDriver::InitListen(FNetworkNotify* InNotify, FURL& ListenURL, bool
 	if (SteamSockets && !ListenURL.HasOption(TEXT("bIsLanMatch")))
 	{
 		FName SocketTypeName = IsRunningDedicatedServer() ? FName(TEXT("SteamServerSocket")) : FName(TEXT("SteamClientSocket"));
-		Socket = SteamSockets->CreateSocket(SocketTypeName, TEXT("Unreal server (Steam)"));
+		Socket = SteamSockets->CreateSocket(SocketTypeName, TEXT("Unreal server (Steam)"), FNetworkProtocolTypes::Steam);
 	}
 	else
 	{
diff --git a/Engine/Plugins/Online/OnlineSubsystemUtils/Source/OnlineSubsystemUtils/Classes/IpConnection.h b/Engine/Plugins/Online/OnlineSubsystemUtils/Source/OnlineSubsystemUtils/Classes/IpConnection.h
--- a/Engine/Plugins/Online/OnlineSubsystemUtils/Source/OnlineSubsystemUtils/Classes/IpConnection.h
+++ b/Engine/Plugins/Online/OnlineSubsystemUtils/Source/OnlineSubsystemUtils/Classes/IpConnection.h
@@ -22,10 +22,6 @@ class ONLINESUBSYSTEMUTILS_API UIpConnection : public UNetConnection
     GENERATED_UCLASS_BODY()
 	// Variables.
 
-	// @todo #JIRA UENET-883: This should be moved down to UNetConnection, now that GetInternetAddr is a thing.
-	//			A lot of platforms reinvent the wheel in their own inefficient way here, despite having their own FInternetAddr type
-	TSharedPtr<FInternetAddr>	RemoteAddr;
-
 	class FSocket*				Socket;
 	class FResolveInfo*			ResolveInfo;
 
@@ -36,10 +32,6 @@ class ONLINESUBSYSTEMUTILS_API UIpConnection : public UNetConnection
 	virtual void LowLevelSend(void* Data, int32 CountBits, FOutPacketTraits& Traits) override;
 	FString LowLevelGetRemoteAddress(bool bAppendPort=false) override;
 	FString LowLevelDescribe() override;
-	virtual int32 GetAddrAsInt(void) override;
-	virtual int32 GetAddrPort(void) override;
-	virtual TSharedPtr<FInternetAddr> GetInternetAddr() override;
-	virtual FString RemoteAddressToString() override;
 	virtual void Tick() override;
 	virtual void CleanUp() override;
 	//~ End NetConnection Interface
diff --git a/Engine/Plugins/Online/OnlineSubsystemUtils/Source/OnlineSubsystemUtils/Classes/IpNetDriver.h b/Engine/Plugins/Online/OnlineSubsystemUtils/Source/OnlineSubsystemUtils/Classes/IpNetDriver.h
--- a/Engine/Plugins/Online/OnlineSubsystemUtils/Source/OnlineSubsystemUtils/Classes/IpNetDriver.h
+++ b/Engine/Plugins/Online/OnlineSubsystemUtils/Source/OnlineSubsystemUtils/Classes/IpNetDriver.h
@@ -37,9 +37,6 @@ class ONLINESUBSYSTEMUTILS_API UIpNetDriver : public UNetDriver
 	UPROPERTY(Config)
 	uint32 MaxPortCountToTry;
 
-	/** Local address this net driver is associated with */
-	TSharedPtr<FInternetAddr> LocalAddr;
-
 	/** Underlying socket communication */
 	FSocket* Socket;
 
@@ -52,7 +49,7 @@ class ONLINESUBSYSTEMUTILS_API UIpNetDriver : public UNetDriver
 	virtual bool InitConnect( FNetworkNotify* InNotify, const FURL& ConnectURL, FString& Error ) override;
 	virtual bool InitListen( FNetworkNotify* InNotify, FURL& LocalURL, bool bReuseAddressAndPort, FString& Error ) override;
 	virtual void TickDispatch( float DeltaTime ) override;
-	virtual void LowLevelSend(FString Address, void* Data, int32 CountBits, FOutPacketTraits& Traits) override;
+	virtual void LowLevelSend(TSharedPtr<const FInternetAddr> Address, void* Data, int32 CountBits, FOutPacketTraits& Traits) override;
 	virtual FString LowLevelGetNetworkNumber() override;
 	virtual void LowLevelDestroy() override;
 	virtual class ISocketSubsystem* GetSocketSubsystem() override;
diff --git a/Engine/Plugins/Online/OnlineSubsystemUtils/Source/OnlineSubsystemUtils/Private/IpConnection.cpp b/Engine/Plugins/Online/OnlineSubsystemUtils/Source/OnlineSubsystemUtils/Private/IpConnection.cpp
--- a/Engine/Plugins/Online/OnlineSubsystemUtils/Source/OnlineSubsystemUtils/Private/IpConnection.cpp
+++ b/Engine/Plugins/Online/OnlineSubsystemUtils/Source/OnlineSubsystemUtils/Private/IpConnection.cpp
@@ -37,7 +37,6 @@ TAutoConsoleVariable<int32> CVarNetIpConnectionUseSendTasks(
 
 UIpConnection::UIpConnection(const FObjectInitializer& ObjectInitializer) :
 	Super(ObjectInitializer),
-	RemoteAddr(NULL),
 	Socket(NULL),
 	ResolveInfo(NULL)
 {
@@ -94,13 +93,7 @@ void UIpConnection::InitRemoteConnection(UNetDriver* InDriver, class FSocket* In
 		(InMaxPacket == 0 || InMaxPacket > MAX_PACKET_SIZE) ? MAX_PACKET_SIZE : InMaxPacket,
 		InPacketOverhead == 0 ? UDP_HEADER_SIZE : InPacketOverhead);
 
-	// Copy the remote IPAddress passed in
-	bool bIsValid = false;
-	FString IpAddrStr = InRemoteAddr.ToString(false);
-	RemoteAddr = InDriver->GetSocketSubsystem()->CreateInternetAddr();
-	RemoteAddr->SetIp(*IpAddrStr, bIsValid);
-	RemoteAddr->SetPort(InRemoteAddr.GetPort());
-
+	RemoteAddr = InRemoteAddr.Clone();
 	URL.Host = RemoteAddr->ToString(false);
 
 	// Initialize our send bunch
@@ -332,26 +325,3 @@ FString UIpConnection::LowLevelDescribe()
 		:								TEXT("Invalid")
 	);
 }
-
-int32 UIpConnection::GetAddrAsInt(void)
-{
-	uint32 OutAddr = 0;
-	// Get the host byte order ip addr
-	RemoteAddr->GetIp(OutAddr);
-	return (int32)OutAddr;
-}
-
-int32 UIpConnection::GetAddrPort(void)
-{
-	return RemoteAddr->GetPort();
-}
-
-TSharedPtr<FInternetAddr> UIpConnection::GetInternetAddr()
-{
-	return RemoteAddr;
-}
-
-FString UIpConnection::RemoteAddressToString()
-{
-	return RemoteAddr->ToString(true);
-}
diff --git a/Engine/Plugins/Online/OnlineSubsystemUtils/Source/OnlineSubsystemUtils/Private/IpNetDriver.cpp b/Engine/Plugins/Online/OnlineSubsystemUtils/Source/OnlineSubsystemUtils/Private/IpNetDriver.cpp
--- a/Engine/Plugins/Online/OnlineSubsystemUtils/Source/OnlineSubsystemUtils/Private/IpNetDriver.cpp
+++ b/Engine/Plugins/Online/OnlineSubsystemUtils/Source/OnlineSubsystemUtils/Private/IpNetDriver.cpp
@@ -106,7 +106,7 @@ FSocket * UIpNetDriver::CreateSocket()
 		return NULL;
 	}
 
-	return SocketSubsystem->CreateSocket( NAME_DGram, TEXT( "Unreal" ) );
+	return SocketSubsystem->CreateSocket(NAME_DGram, TEXT("Unreal"), ( LocalAddr.IsValid() ? LocalAddr->GetProtocolType() : NAME_None) );
 }
 
 int UIpNetDriver::GetClientPort()
@@ -122,18 +122,32 @@ bool UIpNetDriver::InitBase( bool bInitAsClient, FNetworkNotify* InNotify, const
 	}
 
 	ISocketSubsystem* SocketSubsystem = GetSocketSubsystem();
-	if (SocketSubsystem == NULL)
+	if (SocketSubsystem == nullptr)
 	{
 		UE_LOG(LogNet, Warning, TEXT("Unable to find socket subsystem"));
 		return false;
 	}
 
 	// Derived types may have already allocated a socket
+	const TCHAR* MultiHomeBindAddr = URL.GetOption(TEXT("multihome="), nullptr);
+	if (MultiHomeBindAddr != nullptr)
+	{
+		LocalAddr = SocketSubsystem->GetAddressFromString(MultiHomeBindAddr);
+		if (!LocalAddr.IsValid())
+		{
+			UE_LOG(LogNet, Warning, TEXT("Failed to created valid address from multihome address: %s"), MultiHomeBindAddr);
+		}
+	}
+
+	if (!LocalAddr.IsValid())
+	{
+		LocalAddr = SocketSubsystem->GetLocalBindAddr(*GLog);
+	}
 
 	// Create the socket that we will use to communicate with
 	Socket = CreateSocket();
 
-	if( Socket == NULL )
+	if(Socket == nullptr)
 	{
 		Socket = 0;
 		Error = FString::Printf( TEXT("%s: socket failed (%i)"), SocketSubsystem->GetSocketAPIName(), (int32)SocketSubsystem->GetLastErrorCode() );
@@ -164,27 +178,6 @@ bool UIpNetDriver::InitBase( bool bInitAsClient, FNetworkNotify* InNotify, const
 	Socket->SetSendBufferSize(SendSize,SendSize);
 	UE_LOG(LogInit, Log, TEXT("%s: Socket queue %i / %i"), SocketSubsystem->GetSocketAPIName(), RecvSize, SendSize );
 
-	// allow multihome as a url option
-	const TCHAR* MultiHomeBindAddr = URL.GetOption(TEXT("multihome="), nullptr);
-	if (MultiHomeBindAddr != nullptr)
-	{
-		LocalAddr = SocketSubsystem->CreateInternetAddr();
-
-		bool bIsValid = false;
-		LocalAddr->SetIp(MultiHomeBindAddr, bIsValid);
-
-		if (!bIsValid)
-		{
-			LocalAddr = nullptr;
-			UE_LOG(LogNet, Warning, TEXT("Failed to created valid address from multihome address: %s"), MultiHomeBindAddr);
-		}
-	}
-
-	if (!LocalAddr.IsValid())
-	{
-		LocalAddr = SocketSubsystem->GetLocalBindAddr(*GLog);
-	}
-
 	// Bind socket to our port.
 	LocalAddr->SetPort(bInitAsClient ? GetClientPort() : URL.Port);
 	
@@ -456,7 +449,8 @@ void UIpNetDriver::TickDispatch(float DeltaTime)
 
 		if (Connection == nullptr)
 		{
-			UNetConnection** Result = MappedClientConnections.Find(FromAddr);
+			const TSharedRef<const FInternetAddr> ConstFromAddr = FromAddr;
+			UNetConnection** Result = MappedClientConnections.Find(ConstFromAddr);
 
 			if (Result != nullptr)
 			{
@@ -589,11 +583,11 @@ UNetConnection* UIpNetDriver::ProcessConnectionlessPacket(const TSharedRef<FInte
 	{
 		StatelessConnect = StatelessConnectComponent.Pin();
 
-		const ProcessedPacket UnProcessedPacket = ConnectionlessHandler->IncomingConnectionless(IncomingAddress, Data, CountBytesRef);
+		const ProcessedPacket UnProcessedPacket = ConnectionlessHandler->IncomingConnectionless(Address, Data, CountBytesRef);
 
 		if (!UnProcessedPacket.bError)
 		{
-			bPassedChallenge = StatelessConnect->HasPassedChallenge(IncomingAddress, bRestartedHandshake);
+			bPassedChallenge = StatelessConnect->HasPassedChallenge(Address, bRestartedHandshake);
 
 			if (bPassedChallenge)
 			{
@@ -735,17 +729,9 @@ UNetConnection* UIpNetDriver::ProcessConnectionlessPacket(const TSharedRef<FInte
 	return ReturnVal;
 }
 
-void UIpNetDriver::LowLevelSend(FString Address, void* Data, int32 CountBits, FOutPacketTraits& Traits)
+void UIpNetDriver::LowLevelSend(TSharedPtr<const FInternetAddr> Address, void* Data, int32 CountBits, FOutPacketTraits& Traits)
 {
-	bool bValidAddress = !Address.IsEmpty();
-	TSharedRef<FInternetAddr> RemoteAddr = GetSocketSubsystem()->CreateInternetAddr();
-
-	if (bValidAddress)
-	{
-		RemoteAddr->SetIp(*Address, bValidAddress);
-	}
-
-	if (bValidAddress)
+	if (Address.IsValid() && Address->IsValid())
 	{
 		const uint8* DataToSend = reinterpret_cast<uint8*>(Data);
 
@@ -771,7 +757,7 @@ void UIpNetDriver::LowLevelSend(FString Address, void* Data, int32 CountBits, FO
 		if (CountBits > 0)
 		{
 			CLOCK_CYCLES(SendCycles);
-			Socket->SendTo(DataToSend, FMath::DivideAndRoundUp(CountBits, 8), BytesSent, *RemoteAddr);
+			Socket->SendTo(DataToSend, FMath::DivideAndRoundUp(CountBits, 8), BytesSent, *Address);
 			UNCLOCK_CYCLES(SendCycles);
 		}
 
@@ -783,7 +769,7 @@ void UIpNetDriver::LowLevelSend(FString Address, void* Data, int32 CountBits, FO
 	}
 	else
 	{
-		UE_LOG(LogNet, Warning, TEXT("UIpNetDriver::LowLevelSend: Invalid send address '%s'"), *Address);
+		UE_LOG(LogNet, Warning, TEXT("UIpNetDriver::LowLevelSend: Invalid send address '%s'"), *Address->ToString(true));
 	}
 }
 
@@ -791,7 +777,7 @@ void UIpNetDriver::LowLevelSend(FString Address, void* Data, int32 CountBits, FO
 
 FString UIpNetDriver::LowLevelGetNetworkNumber()
 {
-	return LocalAddr->ToString(true);
+	return LocalAddr.IsValid() ? LocalAddr->ToString(true) : FString(TEXT(""));
 }
 
 void UIpNetDriver::LowLevelDestroy()
diff --git a/Engine/Plugins/Online/OnlineSubsystemUtils/Source/OnlineSubsystemUtils/Private/OnlineSubsystemUtils.cpp b/Engine/Plugins/Online/OnlineSubsystemUtils/Source/OnlineSubsystemUtils/Private/OnlineSubsystemUtils.cpp
--- a/Engine/Plugins/Online/OnlineSubsystemUtils/Source/OnlineSubsystemUtils/Private/OnlineSubsystemUtils.cpp
+++ b/Engine/Plugins/Online/OnlineSubsystemUtils/Source/OnlineSubsystemUtils/Private/OnlineSubsystemUtils.cpp
@@ -218,7 +218,9 @@ int32 GetClientPeerIp(FName InstanceName, const FUniqueNetId& UserId)
 				if (ClientConnection && 
 					ClientConnection->PlayerId.ToString() == UserId.ToString())
 				{
+					PRAGMA_DISABLE_DEPRECATION_WARNINGS
 					PeerIp = ClientConnection->GetAddrAsInt();
+					PRAGMA_ENABLE_DEPRECATION_WARNINGS
 					break;
 				}
 			}
diff --git a/Engine/Plugins/Runtime/PacketHandlers/AESHandlerComponent/Source/Private/AESHandlerComponent.cpp b/Engine/Plugins/Runtime/PacketHandlers/AESHandlerComponent/Source/Private/AESHandlerComponent.cpp
--- a/Engine/Plugins/Runtime/PacketHandlers/AESHandlerComponent/Source/Private/AESHandlerComponent.cpp
+++ b/Engine/Plugins/Runtime/PacketHandlers/AESHandlerComponent/Source/Private/AESHandlerComponent.cpp
@@ -81,7 +81,7 @@ void FAESHandlerComponent::Incoming(FBitReader& Packet)
 			// Copy remaining bits to a TArray so that they are byte-aligned.
 			Ciphertext.Reset();
 			Ciphertext.AddUninitialized(Packet.GetBytesLeft());
-			Ciphertext[Ciphertext.Num()-1] = 0;
+			Ciphertext[Ciphertext.Num() - 1] = 0;
 
 			Packet.SerializeBits(Ciphertext.GetData(), Packet.GetBitsLeft());
 
@@ -194,14 +194,6 @@ void FAESHandlerComponent::Outgoing(FBitWriter& Packet, FOutPacketTraits& Traits
 	}
 }
 
-void FAESHandlerComponent::IncomingConnectionless(const FString& Address, FBitReader& Packet)
-{
-}
-
-void FAESHandlerComponent::OutgoingConnectionless(const FString& Address, FBitWriter& Packet, FOutPacketTraits& Traits)
-{
-}
-
 int32 FAESHandlerComponent::GetReservedPacketBits() const
 {
 	// Worst case includes the encryption enabled bit, the termination bit, padding up to the next whole byte, and a block of padding.
diff --git a/Engine/Plugins/Runtime/PacketHandlers/AESHandlerComponent/Source/Public/AESHandlerComponent.h b/Engine/Plugins/Runtime/PacketHandlers/AESHandlerComponent/Source/Public/AESHandlerComponent.h
--- a/Engine/Plugins/Runtime/PacketHandlers/AESHandlerComponent/Source/Public/AESHandlerComponent.h
+++ b/Engine/Plugins/Runtime/PacketHandlers/AESHandlerComponent/Source/Public/AESHandlerComponent.h
@@ -44,8 +44,8 @@ public:
 	virtual bool IsValid() const override;
 	virtual void Incoming(FBitReader& Packet) override;
 	virtual void Outgoing(FBitWriter& Packet, FOutPacketTraits& Traits) override;
-	virtual void IncomingConnectionless(const FString& Address, FBitReader& Packet) override;
-	virtual void OutgoingConnectionless(const FString& Address, FBitWriter& Packet, FOutPacketTraits& Traits) override;
+	virtual void IncomingConnectionless(const TSharedPtr<const FInternetAddr>& Address, FBitReader& Packet) override {}
+	virtual void OutgoingConnectionless(const TSharedPtr<const FInternetAddr>& Address, FBitWriter& Packet, FOutPacketTraits& Traits) override {}
 	virtual int32 GetReservedPacketBits() const override;
 	virtual void CountBytes(FArchive& Ar) const override;
 
diff --git a/Engine/Plugins/Runtime/PacketHandlers/CompressionComponents/Oodle/Source/OodleHandlerComponent/Public/OodleHandlerComponent.h b/Engine/Plugins/Runtime/PacketHandlers/CompressionComponents/Oodle/Source/OodleHandlerComponent/Public/OodleHandlerComponent.h
--- a/Engine/Plugins/Runtime/PacketHandlers/CompressionComponents/Oodle/Source/OodleHandlerComponent/Public/OodleHandlerComponent.h
+++ b/Engine/Plugins/Runtime/PacketHandlers/CompressionComponents/Oodle/Source/OodleHandlerComponent/Public/OodleHandlerComponent.h
@@ -344,11 +344,11 @@ public:
 
 	virtual void Outgoing(FBitWriter& Packet, FOutPacketTraits& Traits) override;
 
-	virtual void IncomingConnectionless(const FString& Address, FBitReader& Packet) override
+	virtual void IncomingConnectionless(const TSharedPtr<const FInternetAddr>& Address, FBitReader& Packet) override
 	{
 	}
 
-	virtual void OutgoingConnectionless(const FString& Address, FBitWriter& Packet, FOutPacketTraits& Traits) override
+	virtual void OutgoingConnectionless(const TSharedPtr<const FInternetAddr>& Address, FBitWriter& Packet, FOutPacketTraits& Traits) override
 	{
 	}
 
diff --git a/Engine/Source/Runtime/Engine/Classes/Engine/DemoNetConnection.h b/Engine/Source/Runtime/Engine/Classes/Engine/DemoNetConnection.h
--- a/Engine/Source/Runtime/Engine/Classes/Engine/DemoNetConnection.h
+++ b/Engine/Source/Runtime/Engine/Classes/Engine/DemoNetConnection.h
@@ -75,7 +75,7 @@ public:
 	virtual int32 IsNetReady( bool Saturate ) override;
 	virtual void FlushNet( bool bIgnoreSimulation = false ) override;
 	virtual void HandleClientPlayer( APlayerController* PC, class UNetConnection* NetConnection ) override;
-	virtual TSharedPtr<FInternetAddr> GetInternetAddr() override;
+	virtual TSharedPtr<const FInternetAddr> GetRemoteAddr() override { return nullptr; }
 	virtual bool ClientHasInitializedLevelFor( const AActor* TestActor ) const override;
 	virtual TSharedPtr<FObjectReplicator> CreateReplicatorForNewActorChannel(UObject* Object);
 	virtual FString RemoteAddressToString() override { return TEXT("Demo"); }
diff --git a/Engine/Source/Runtime/Engine/Classes/Engine/NetConnection.h b/Engine/Source/Runtime/Engine/Classes/Engine/NetConnection.h
--- a/Engine/Source/Runtime/Engine/Classes/Engine/NetConnection.h
+++ b/Engine/Source/Runtime/Engine/Classes/Engine/NetConnection.h
@@ -274,6 +274,9 @@ class UNetConnection : public UPlayer
 	uint32 InternalAck:1;					// Internally ack all packets, for 100% reliable connections.
 
 	struct FURL			URL;				// URL of the other side.
+	
+	/** The remote address of this connection, typically generated from the URL. */
+	TSharedPtr<FInternetAddr>	RemoteAddr;
 
 	// Track each type of bit used per-packet for bandwidth profiling
 
@@ -741,14 +744,28 @@ public:
 	ENGINE_API virtual void HandleClientPlayer( class APlayerController* PC, class UNetConnection* NetConnection );
 
 	/** @return the address of the connection as an integer */
+	UE_DEPRECATED(4.23, "Use GetRemoteAddr as it allows direct access to the RemoteAddr and allows for dynamic address sizing.")
 	virtual int32 GetAddrAsInt(void)
 	{
+		PRAGMA_DISABLE_DEPRECATION_WARNINGS
+		if (RemoteAddr.IsValid())
+		{
+			uint32 OutAddr = 0;
+			// Get the host byte order ip addr
+			RemoteAddr->GetIp(OutAddr);
+			return (int32)OutAddr;
+		}
+		PRAGMA_ENABLE_DEPRECATION_WARNINGS
 		return 0;
 	}
 
 	/** @return the port of the connection as an integer */
 	virtual int32 GetAddrPort(void)
 	{
+		if (RemoteAddr.IsValid())
+		{
+			return RemoteAddr->GetPort();
+		}
 		return 0;
 	}
 
@@ -758,7 +775,16 @@ public:
 	 *
 	 * @return	The platform specific FInternetAddr containing this connections address
 	 */
-	virtual TSharedPtr<FInternetAddr> GetInternetAddr() PURE_VIRTUAL(UNetConnection::GetInternetAddr,return TSharedPtr<FInternetAddr>(););
+	UE_DEPRECATED(4.23, "Use GetRemoteAddr to safely get the FInternetAddr tied to this connection")
+	virtual TSharedPtr<FInternetAddr> GetInternetAddr() { return ConstCastSharedPtr<FInternetAddr>(GetRemoteAddr()); }
+
+	/**
+	 * Return the platform specific FInternetAddr type, containing this connections address.
+	 * If nullptr is returned, connection is not added to MappedClientConnections, and can't receive net packets which depend on this.
+	 *
+	 * @return	The platform specific FInternetAddr containing this connections address
+	 */
+	virtual TSharedPtr<const FInternetAddr> GetRemoteAddr() { return RemoteAddr; }
 
 	/** closes the connection (including sending a close notify across the network) */
 	ENGINE_API void Close();
@@ -873,7 +899,14 @@ public:
 	* Gets a unique ID for the connection, this ID depends on the underlying connection
 	* For IP connections this is an IP Address and port, for steam this is a SteamID
 	*/
-	ENGINE_API virtual FString RemoteAddressToString() PURE_VIRTUAL(UNetConnection::RemoteAddressToString, return TEXT("Error"););
+	ENGINE_API virtual FString RemoteAddressToString()
+	{
+		if (RemoteAddr.IsValid())
+		{
+			return RemoteAddr->ToString(true);
+		}
+		return TEXT("Invalid");
+	}
 	
 	
 	/** Called by UActorChannel. Handles creating a new replicator for an actor */
@@ -1194,7 +1227,6 @@ public:
 	virtual FString LowLevelGetRemoteAddress(bool bAppendPort=false) override { return FString(); }
 	virtual bool ClientHasInitializedLevelFor(const AActor* TestActor) const { return true; }
 
-
-	virtual TSharedPtr<FInternetAddr> GetInternetAddr() override { return TSharedPtr<FInternetAddr>(); }
+	virtual TSharedPtr<const FInternetAddr> GetRemoteAddr() override { return nullptr; }
 };
 
diff --git a/Engine/Source/Runtime/Engine/Classes/Engine/NetDriver.h b/Engine/Source/Runtime/Engine/Classes/Engine/NetDriver.h
--- a/Engine/Source/Runtime/Engine/Classes/Engine/NetDriver.h
+++ b/Engine/Source/Runtime/Engine/Classes/Engine/NetDriver.h
@@ -35,7 +35,7 @@ class UChannel;
 class IAnalyticsProvider;
 class FNetAnalyticsAggregator;
 
-using FConnectionMap = TMap<TSharedRef<FInternetAddr>, UNetConnection*, FDefaultSetAllocator, FInternetAddrKeyMapFuncs<UNetConnection*>>;
+using FConnectionMap = TMap<TSharedRef<const FInternetAddr>, UNetConnection*, FDefaultSetAllocator, FInternetAddrConstKeyMapFuncs<UNetConnection*>>;
 
 extern ENGINE_API TAutoConsoleVariable<int32> CVarNetAllowEncryption;
 extern ENGINE_API int32 GNumSaturatedConnections;
@@ -291,11 +291,18 @@ struct ENGINE_API FChannelDefinition
 struct FDisconnectedClient
 {
 	/** The address of the client */
-	TSharedRef<FInternetAddr>	Address;
+	TSharedRef<const FInternetAddr>	Address;
 
 	/** The time at which the client disconnected  */
 	double						DisconnectTime;
 
+	FDisconnectedClient(TSharedRef<const FInternetAddr>& InAddress, double InDisconnectTime)
+		: Address(InAddress)
+		, DisconnectTime(InDisconnectTime)
+	{
+	}
+
+	UE_DEPRECATED(4.23, "This object does not update address information, and should have const addresses passed in")
 	FDisconnectedClient(TSharedRef<FInternetAddr>& InAddress, double InDisconnectTime)
 		: Address(InAddress)
 		, DisconnectTime(InDisconnectTime)
@@ -726,7 +733,8 @@ public:
 	/** DDoS detection management */
 	FDDoSDetection DDoS;
 
-
+	/** Local address this net driver is associated with */
+	TSharedPtr<FInternetAddr> LocalAddr;
 
 	/**
 	* Updates the standby cheat information and
@@ -845,7 +853,10 @@ public:
 	ENGINE_API virtual void LowLevelDestroy();
 
 	/* @return network number */
-	virtual FString LowLevelGetNetworkNumber() PURE_VIRTUAL(UNetDriver::LowLevelGetNetworkNumber,return TEXT(""););
+	virtual FString LowLevelGetNetworkNumber();
+
+	/* @return local addr of this machine if set */
+	virtual TSharedPtr<const FInternetAddr> GetLocalAddr() { return LocalAddr; }
 
 	/** Make sure this connection is in a reasonable state. */
 	ENGINE_API virtual void AssertValid();
@@ -910,9 +921,14 @@ public:
 	ENGINE_API virtual void LowLevelSend(FString Address, void* Data, int32 CountBits)
 	{
 		FOutPacketTraits EmptyTraits;
+		PRAGMA_DISABLE_DEPRECATION_WARNINGS
 		LowLevelSend(Address, Data, CountBits, EmptyTraits);
+		PRAGMA_ENABLE_DEPRECATION_WARNINGS
 	}
 
+	UE_DEPRECATED(4.22, "Change arguments to support FInternetAddr instead")
+	ENGINE_API virtual void LowLevelSend(FString Address, void* Data, int32 CountBits, FOutPacketTraits& Traits);
+
 	/**
 	 * Sends a 'connectionless' (not associated with a UNetConection) packet, to the specified address.
 	 * NOTE: Address is an abstract format defined by subclasses. Anything calling this, must use an address supplied by the net driver.
@@ -922,7 +938,7 @@ public:
 	 * @param CountBits		The size of the packet data, in bits
 	 * @param Traits		Traits for the packet, if applicable
 	 */
-	ENGINE_API virtual void LowLevelSend(FString Address, void* Data, int32 CountBits, FOutPacketTraits& Traits)
+	ENGINE_API virtual void LowLevelSend(TSharedPtr<const FInternetAddr> Address, void* Data, int32 CountBits, FOutPacketTraits& Traits)
 		PURE_VIRTUAL(UNetDriver::LowLevelSend,);
 
 	/**
diff --git a/Engine/Source/Runtime/Engine/Private/DemoNetDriver.cpp b/Engine/Source/Runtime/Engine/Private/DemoNetDriver.cpp
--- a/Engine/Source/Runtime/Engine/Private/DemoNetDriver.cpp
+++ b/Engine/Source/Runtime/Engine/Private/DemoNetDriver.cpp
@@ -5225,12 +5225,6 @@ void UDemoNetConnection::HandleClientPlayer( APlayerController* PC, UNetConnecti
 	}
 }
 
-TSharedPtr<FInternetAddr> UDemoNetConnection::GetInternetAddr()
-{
-	// Does not use MappedClientConnections
-	return TSharedPtr<FInternetAddr>();
-}
-
 bool UDemoNetConnection::ClientHasInitializedLevelFor(const AActor* TestActor) const
 {
 	// We save all currently streamed levels into the demo stream so we can force the demo playback client
diff --git a/Engine/Source/Runtime/Engine/Private/NetConnection.cpp b/Engine/Source/Runtime/Engine/Private/NetConnection.cpp
--- a/Engine/Source/Runtime/Engine/Private/NetConnection.cpp
+++ b/Engine/Source/Runtime/Engine/Private/NetConnection.cpp
@@ -34,6 +34,7 @@
 #include "UObject/ObjectKey.h"
 #include "UObject/UObjectIterator.h"
 #include "Net/NetworkGranularMemoryLogging.h"
+#include "SocketSubsystem.h"
 
 static TAutoConsoleVariable<int32> CVarPingExcludeFrameTime( TEXT( "net.PingExcludeFrameTime" ), 0, TEXT( "Calculate RTT time between NIC's of server and client." ) );
 
@@ -97,6 +98,7 @@ UNetConnection::UNetConnection(const FObjectInitializer& ObjectInitializer)
 ,   OwningActor			( nullptr )
 ,	MaxPacket			( 0 )
 ,	InternalAck			( false )
+,	RemoteAddr			( nullptr )
 ,	MaxPacketHandlerBits ( 0 )
 ,	State				( USOCK_Invalid )
 ,	Handler()
@@ -320,6 +322,11 @@ void UNetConnection::InitHandler()
 		{
 			Handler::Mode Mode = Driver->ServerConnection != nullptr ? Handler::Mode::Client : Handler::Mode::Server;
 
+			PRAGMA_DISABLE_DEPRECATION_WARNINGS
+			Handler->InitializeAddressSerializer([this](const FString& InAddress){
+				return Driver->GetSocketSubsystem()->GetAddressFromString(InAddress);
+			});
+			PRAGMA_ENABLE_DEPRECATION_WARNINGS
 			Handler->InitializeDelegates(FPacketHandlerLowLevelSendTraits::CreateUObject(this, &UNetConnection::LowLevelSend));
 			Handler->NotifyAnalyticsProvider(Driver->AnalyticsProvider, Driver->AnalyticsAggregator);
 			Handler->Initialize(Mode, MaxPacket * 8, false, nullptr, nullptr, Driver->NetDriverName);
diff --git a/Engine/Source/Runtime/Engine/Private/NetworkDriver.cpp b/Engine/Source/Runtime/Engine/Private/NetworkDriver.cpp
--- a/Engine/Source/Runtime/Engine/Private/NetworkDriver.cpp
+++ b/Engine/Source/Runtime/Engine/Private/NetworkDriver.cpp
@@ -63,6 +63,8 @@
 #include "ProfilingDebugging/CsvProfiler.h"
 #include "Engine/LevelScriptActor.h"
 #include "Net/NetworkGranularMemoryLogging.h"
+#include "SocketSubsystem.h"
+#include "AddressInfoTypes.h"
 
 #if USE_SERVER_PERF_COUNTERS
 #include "PerfCountersModule.h"
@@ -266,6 +268,7 @@ UNetDriver::UNetDriver(const FObjectInitializer& ObjectInitializer)
 #endif
 ,	ProcessQueuedBunchesCurrentFrameMilliseconds(0.0f)
 ,	DDoS()
+,	LocalAddr(nullptr)
 ,	NetworkObjects(new FNetworkObjectList)
 ,	LagState(ENetworkLagState::NotLagging)
 ,	DuplicateLevelID(INDEX_NONE)
@@ -1314,6 +1317,22 @@ void UNetDriver::PostTickFlush()
 	}
 }
 
+PRAGMA_DISABLE_DEPRECATION_WARNINGS
+void UNetDriver::LowLevelSend(FString Address, void* Data, int32 CountBits, FOutPacketTraits& Traits)
+{
+	if (GetSocketSubsystem() != nullptr)
+	{
+		TSharedPtr<FInternetAddr> NetAddress = GetSocketSubsystem()->GetAddressFromString(*Address);
+		if (NetAddress.IsValid())
+		{
+			LowLevelSend(NetAddress, Data, CountBits, Traits);
+			return;
+		}
+	}
+	UE_LOG(LogNet, Warning, TEXT("Could not infer the address to use for LowLevelSend, please use the one that takes an FInternetAddr to avoid this problem"));
+}
+PRAGMA_ENABLE_DEPRECATION_WARNINGS
+
 bool UNetDriver::InitConnectionClass(void)
 {
 	if (NetConnectionClass == NULL && NetConnectionClassName != TEXT(""))
@@ -1403,6 +1422,11 @@ void UNetDriver::InitConnectionlessHandler()
 
 		if (ConnectionlessHandler.IsValid())
 		{
+			PRAGMA_DISABLE_DEPRECATION_WARNINGS
+			ConnectionlessHandler->InitializeAddressSerializer([this](const FString& InAddress) {
+				return GetSocketSubsystem()->GetAddressFromString(InAddress);
+			});
+			PRAGMA_ENABLE_DEPRECATION_WARNINGS
 			ConnectionlessHandler->NotifyAnalyticsProvider(AnalyticsProvider, AnalyticsAggregator);
 			ConnectionlessHandler->Initialize(Handler::Mode::Server, MAX_PACKET_SIZE, true, nullptr, nullptr, NetDriverName);
 
@@ -2389,6 +2413,10 @@ void UNetDriver::LowLevelDestroy()
 	SetWorld(NULL);
 }
 
+FString UNetDriver::LowLevelGetNetworkNumber()
+{
+	return LocalAddr.IsValid() ? LocalAddr->ToString(true) : FString(TEXT(""));
+}
 
 #if !UE_BUILD_SHIPPING
 bool UNetDriver::HandleSocketsCommand( const TCHAR* Cmd, FOutputDevice& Ar )
@@ -4742,20 +4770,17 @@ void UNetDriver::AddClientConnection(UNetConnection * NewConnection)
 
 	ClientConnections.Add(NewConnection);
 
-	TSharedPtr<FInternetAddr> ConnAddr = NewConnection->GetInternetAddr();
+	TSharedPtr<const FInternetAddr> ConnAddr = NewConnection->GetRemoteAddr();
 
 	if (ConnAddr.IsValid())
 	{
-		TSharedRef<FInternetAddr> ConnAddrRef = ConnAddr.ToSharedRef();
-
-		MappedClientConnections.Add(ConnAddrRef, NewConnection);
-
+		MappedClientConnections.Add(ConnAddr.ToSharedRef(), NewConnection);
 
 		// On the off-chance of the same IP:Port being reused, check RecentlyDisconnectedClients
 		int32 RecentDisconnectIdx = RecentlyDisconnectedClients.IndexOfByPredicate(
-			[&ConnAddrRef](const FDisconnectedClient& CurElement)
+			[&ConnAddr](const FDisconnectedClient& CurElement)
 			{
-				return *ConnAddrRef == *CurElement.Address;
+				return *ConnAddr == *CurElement.Address;
 			});
 
 		if (RecentDisconnectIdx != INDEX_NONE)
@@ -4867,28 +4892,28 @@ void UNetDriver::RemoveClientConnection(UNetConnection* ClientConnectionToRemove
 {
 	verify(ClientConnections.Remove(ClientConnectionToRemove) == 1);
 
-	TSharedPtr<FInternetAddr> AddrToRemove = ClientConnectionToRemove->GetInternetAddr();
+	TSharedPtr<const FInternetAddr> AddrToRemove = ClientConnectionToRemove->GetRemoteAddr();
 
 	if (AddrToRemove.IsValid())
 	{
-		TSharedRef<FInternetAddr> AddrRef = AddrToRemove.ToSharedRef();
+		TSharedRef<const FInternetAddr> ConstAddrRef = AddrToRemove.ToSharedRef();
 
 		if (RecentlyDisconnectedTrackingTime > 0)
 		{
-			UNetConnection** FoundVal = MappedClientConnections.Find(AddrRef);
+			UNetConnection** FoundVal = MappedClientConnections.Find(ConstAddrRef);
 
 			// Mark recently disconnected clients as nullptr (don't wait for GC), and keep the MappedClientConections entry for a while.
 			// Required for identifying/ignoring packets from recently disconnected clients, with the same performance as for NetConnection's (important for DDoS detection)
 			if (ensure(FoundVal != nullptr))
 			{
-				RecentlyDisconnectedClients.Add(FDisconnectedClient(AddrRef, FPlatformTime::Seconds()));
+				RecentlyDisconnectedClients.Add(FDisconnectedClient(ConstAddrRef, FPlatformTime::Seconds()));
 
 				*FoundVal = nullptr;
 			}
 		}
 		else
 		{
-			verify(MappedClientConnections.Remove(AddrRef) == 1);
+			verify(MappedClientConnections.Remove(ConstAddrRef) == 1);
 		}
 	}
 
diff --git a/Engine/Source/Runtime/Engine/Private/NetworkProfiler.cpp b/Engine/Source/Runtime/Engine/Private/NetworkProfiler.cpp
--- a/Engine/Source/Runtime/Engine/Private/NetworkProfiler.cpp
+++ b/Engine/Source/Runtime/Engine/Private/NetworkProfiler.cpp
@@ -216,7 +216,7 @@ void FNetworkProfiler::SetCurrentConnection( UNetConnection* Connection )
 {
 	if ( bIsTrackingEnabled && Connection != nullptr )
 	{
-		const TSharedPtr<const FInternetAddr> ConnectionAddr = Connection->GetInternetAddr();
+		const TSharedPtr<const FInternetAddr> ConnectionAddr = Connection->GetRemoteAddr();
 		if ( ConnectionAddr.IsValid() )
 		{
 			if ( LastAddress != ConnectionAddr )
diff --git a/Engine/Source/Runtime/Engine/Private/PacketHandlers/StatelessConnectHandlerComponent.cpp b/Engine/Source/Runtime/Engine/Private/PacketHandlers/StatelessConnectHandlerComponent.cpp
--- a/Engine/Source/Runtime/Engine/Private/PacketHandlers/StatelessConnectHandlerComponent.cpp
+++ b/Engine/Source/Runtime/Engine/Private/PacketHandlers/StatelessConnectHandlerComponent.cpp
@@ -193,7 +193,7 @@ StatelessConnectHandlerComponent::StatelessConnectHandlerComponent()
 	, HandshakeSecret()
 	, ActiveSecret(255)
 	, LastSecretUpdateTimestamp(0.f)
-	, LastChallengeSuccessAddress(TEXT(""))
+	, LastChallengeSuccessAddress(nullptr)
 	, LastServerSequence(0)
 	, LastClientSequence(0)
 	, LastClientSendTimestamp(0.0)
@@ -219,8 +219,6 @@ void StatelessConnectHandlerComponent::CountBytes(FArchive& Ar) const
 	{
 		HandshakeSecret[i].CountBytes(Ar);
 	}
-
-	LastChallengeSuccessAddress.CountBytes(Ar);
 }
 
 void StatelessConnectHandlerComponent::NotifyHandshakeBegin()
@@ -286,7 +284,7 @@ void StatelessConnectHandlerComponent::NotifyHandshakeBegin()
 	}
 }
 
-void StatelessConnectHandlerComponent::SendConnectChallenge(const FString& ClientAddress)
+void StatelessConnectHandlerComponent::SendConnectChallenge(TSharedPtr<const FInternetAddr> ClientAddress)
 {
 	if (Driver != nullptr)
 	{
@@ -425,7 +423,7 @@ void StatelessConnectHandlerComponent::SendChallengeResponse(uint8 InSecretId, f
 	}
 }
 
-void StatelessConnectHandlerComponent::SendChallengeAck(const FString& ClientAddress, uint8 InCookie[COOKIE_BYTE_SIZE])
+void StatelessConnectHandlerComponent::SendChallengeAck(TSharedPtr<const FInternetAddr> ClientAddress, uint8 InCookie[COOKIE_BYTE_SIZE])
 {
 	if (Driver != nullptr)
 	{
@@ -489,7 +487,7 @@ void StatelessConnectHandlerComponent::SendChallengeAck(const FString& ClientAdd
 	}
 }
 
-void StatelessConnectHandlerComponent::SendRestartHandshakeRequest(const FString& ClientAddress)
+void StatelessConnectHandlerComponent::SendRestartHandshakeRequest(const TSharedPtr<const FInternetAddr> ClientAddress)
 {
 	if (Driver != nullptr)
 	{
@@ -716,7 +714,7 @@ void StatelessConnectHandlerComponent::Incoming(FBitReader& Packet)
 				// In this codepath, this component is linked to a UNetConnection, and the Last* values below, cache the handshake info.
 #if !UE_BUILD_SHIPPING
 				UE_LOG(LogHandshake, Log, TEXT("Received unexpected post-connect handshake packet - resending ack for LastChallengeSuccessAddress %s and LastCookie %s."),
-						*LastChallengeSuccessAddress, *FString::FromBlob(AuthorisedCookie, COOKIE_BYTE_SIZE));
+						*LastChallengeSuccessAddress->ToString(true), *FString::FromBlob(AuthorisedCookie, COOKIE_BYTE_SIZE));
 #endif
 
 				SendChallengeAck(LastChallengeSuccessAddress, AuthorisedCookie);
@@ -751,11 +749,11 @@ void StatelessConnectHandlerComponent::Outgoing(FBitWriter& Packet, FOutPacketTr
 	Packet = MoveTemp(NewPacket);
 }
 
-void StatelessConnectHandlerComponent::IncomingConnectionless(const FString& Address, FBitReader& Packet)
+void StatelessConnectHandlerComponent::IncomingConnectionless(const TSharedPtr<const FInternetAddr>& Address, FBitReader& Packet)
 {
 	bool bHandshakePacket = !!Packet.ReadBit() && !Packet.IsError();
 
-	LastChallengeSuccessAddress.Empty();
+	LastChallengeSuccessAddress = nullptr;
 
 	if (bHandshakePacket)
 	{
@@ -891,16 +889,17 @@ bool StatelessConnectHandlerComponent::ParseHandshakePacket(FBitReader& Packet,
 	return bValidPacket;
 }
 
-void StatelessConnectHandlerComponent::GenerateCookie(FString ClientAddress, uint8 SecretId, float Timestamp, uint8 (&OutCookie)[20])
+void StatelessConnectHandlerComponent::GenerateCookie(TSharedPtr<const FInternetAddr> ClientAddress, uint8 SecretId, float Timestamp, uint8 (&OutCookie)[20])
 {
 	// @todo #JohnB: Add cpu stats tracking, like what Oodle does upon compression
 	//					NOTE: Being serverside, will only show up in .uprof, not on any 'stat' commands. Still necessary though.
 
 	TArray<uint8> CookieData;
 	FMemoryWriter CookieArc(CookieData);
+	FString ClientAddressString(ClientAddress->ToString(true));
 
 	CookieArc << Timestamp;
-	CookieArc << ClientAddress;
+	CookieArc << ClientAddressString;
 
 	FSHA1::HMACBuffer(HandshakeSecret[!!SecretId].GetData(), SECRET_BYTE_SIZE, CookieData.GetData(), CookieData.Num(), OutCookie);
 }
diff --git a/Engine/Source/Runtime/Engine/Public/PacketHandlers/StatelessConnectHandlerComponent.h b/Engine/Source/Runtime/Engine/Public/PacketHandlers/StatelessConnectHandlerComponent.h
--- a/Engine/Source/Runtime/Engine/Public/PacketHandlers/StatelessConnectHandlerComponent.h
+++ b/Engine/Source/Runtime/Engine/Public/PacketHandlers/StatelessConnectHandlerComponent.h
@@ -4,6 +4,7 @@
 
 #include "CoreMinimal.h"
 #include "PacketHandler.h"
+#include "IPAddress.h"
 
 /** Temporary version-check macro for randomized packet sequence, for backwards-compatibility. */
 #define STATELESSCONNECT_HAS_RANDOM_SEQUENCE 1
@@ -63,7 +64,7 @@ private:
 	 *
 	 * @param ClientAddress		The address of the client to send the challenge to.
 	 */
-	void SendConnectChallenge(const FString& ClientAddress);
+	void SendConnectChallenge(TSharedPtr<const FInternetAddr> ClientAddress);
 
 	/**
 	 * Constructs and sends the handshake challenge response packet, from the client to the server
@@ -80,14 +81,14 @@ private:
 	 * @param ClientAddress		The address of the client to send the ack to.
 	 * @param InCookie			The cookie value to send
 	 */
-	void SendChallengeAck(const FString& ClientAddress, uint8 InCookie[COOKIE_BYTE_SIZE]);
+	void SendChallengeAck(TSharedPtr<const FInternetAddr> ClientAddress, uint8 InCookie[COOKIE_BYTE_SIZE]);
 
 	/**
 	 * Constructs and sends a request to resend the cookie, from the server to the client.
 	 *
 	 * @param ClientAddress		The address of the client to send the request to.
 	 */
-	void SendRestartHandshakeRequest(const FString& ClientAddress);
+	void SendRestartHandshakeRequest(const TSharedPtr<const FInternetAddr> ClientAddress);
 
 
 	/**
@@ -101,14 +102,15 @@ public:
 	/**
 	 * Whether or not the specified connection address, has just passed the connection handshake challenge.
 	 *
-	 * @param Address					The address (including port, for UIpNetDriver) being checked
+	 * @param Address	The address (including port, for UIpNetDriver) being checked
 	 */
 	UE_DEPRECATED(4.22, "HasPassedChallenge must return bOutRestartedHandshake, in order to correctly process the challenge. Use the new version")
 	FORCEINLINE bool HasPassedChallenge(const FString& Address) const
 	{
 		bool bDud = false;
-
+		PRAGMA_DISABLE_DEPRECATION_WARNINGS
 		return HasPassedChallenge(Address, bDud);
+		PRAGMA_ENABLE_DEPRECATION_WARNINGS
 	}
 
 	/**
@@ -117,13 +119,28 @@ public:
 	 * @param Address					The address (including port, for UIpNetDriver) being checked
 	 * @param bOutRestartedHandshake	Whether or not the passed challenge, was a restarted handshake for an existing NetConnection
 	 */
+	UE_DEPRECATED(4.23, "HasPassedChallenge should pass a FInternetAddr pointer for faster address checking")
 	FORCEINLINE bool HasPassedChallenge(const FString& Address, bool& bOutRestartedHandshake) const
 	{
 		bOutRestartedHandshake = bRestartedHandshake;
 
-		return LastChallengeSuccessAddress == Address;
+		return (LastChallengeSuccessAddress.IsValid()) ? LastChallengeSuccessAddress->ToString(true) == Address : false;
+	}
+
+	/**
+	 * Whether or not the specified connection address, has just passed the connection handshake challenge.
+	 *
+	 * @param Address					The address (including port, for UIpNetDriver) being checked
+	 * @param bOutRestartedHandshake	Whether or not the passed challenge, was a restarted handshake for an existing NetConnection
+	 */
+	FORCEINLINE bool HasPassedChallenge(const TSharedPtr<const FInternetAddr>& Address, bool& bOutRestartedHandshake) const
+	{
+		bOutRestartedHandshake = bRestartedHandshake;
+
+		return (LastChallengeSuccessAddress.IsValid() && Address.IsValid()) ? *LastChallengeSuccessAddress == *Address : false;
 	}
 
+
 	/**
 	 * Used for retrieving the initial packet sequence values from the handshake data, after a successful challenge
 	 *
@@ -152,7 +169,7 @@ public:
 	 */
 	FORCEINLINE void ResetChallengeData()
 	{
-		LastChallengeSuccessAddress.Empty();
+		LastChallengeSuccessAddress = nullptr;
 		bRestartedHandshake = false;
 		LastServerSequence = 0;
 		LastClientSequence = 0;
@@ -175,9 +192,9 @@ protected:
 
 	virtual void Outgoing(FBitWriter& Packet, FOutPacketTraits& Traits) override;
 
-	virtual void IncomingConnectionless(const FString& Address, FBitReader& Packet) override;
+	virtual void IncomingConnectionless(const TSharedPtr<const FInternetAddr>& Address, FBitReader& Packet) override;
 
-	virtual void OutgoingConnectionless(const FString& Address, FBitWriter& Packet, FOutPacketTraits& Traits) override
+	virtual void OutgoingConnectionless(const TSharedPtr<const FInternetAddr>& Address, FBitWriter& Packet, FOutPacketTraits& Traits) override
 	{
 	}
 
@@ -194,13 +211,13 @@ private:
 	/**
 	 * Parses an incoming handshake packet (does not parse the handshake bit though)
 	 *
-	 * @param Packet				The packet the handshake is being parsed from
+	 * @param Packet		The packet the handshake is being parsed from
 	 * @param bOutRestartHandshake	Whether or not this packet is a restart handshake packet
-	 * @param OutSecretId			Which of the two serverside HandshakeSecret values this is based on
-	 * @param OutTimestamp			The server timestamp, from the moment the challenge was sent (or 0.f if from the client)
-	 * @param OutCookie				A unique identifier, generated by the server, which the client must reply with (or 0, for initial packet)
+	 * @param OutSecretId	Which of the two serverside HandshakeSecret values this is based on
+	 * @param OutTimestamp	The server timestamp, from the moment the challenge was sent (or 0.f if from the client)
+	 * @param OutCookie		A unique identifier, generated by the server, which the client must reply with (or 0, for initial packet)
 	 * @param OutOrigCookie			If this is a restart handshake challenge response, this is the original handshake's cookie
-	 * @return						Whether or not the handshake packet was parsed successfully
+	 * @return				Whether or not the handshake packet was parsed successfully
 	 */
 	bool ParseHandshakePacket(FBitReader& Packet, bool& bOutRestartHandshake, uint8& OutSecretId, float& OutTimestamp,
 								uint8 (&OutCookie)[COOKIE_BYTE_SIZE], uint8 (&OutOrigCookie)[COOKIE_BYTE_SIZE]);
@@ -213,7 +230,7 @@ private:
 	 * @param TimeStamp			The serverside timestamp
 	 * @param OutCookie			Outputs the generated cookie value.
 	 */
-	void GenerateCookie(FString ClientAddress, uint8 SecretId, float Timestamp, uint8 (&OutCookie)[COOKIE_BYTE_SIZE]);
+	void GenerateCookie(TSharedPtr<const FInternetAddr> ClientAddress, uint8 SecretId, float Timestamp, uint8 (&OutCookie)[COOKIE_BYTE_SIZE]);
 
 	/**
 	 * Generates a new HandshakeSecret value
@@ -238,7 +255,7 @@ private:
 	float LastSecretUpdateTimestamp;
 
 	/** The last address to successfully complete the handshake challenge */
-	FString LastChallengeSuccessAddress;
+	TSharedPtr<const FInternetAddr> LastChallengeSuccessAddress;
 
 	/** The initial server sequence value, from the last successful handshake */
 	int32 LastServerSequence;
diff --git a/Engine/Source/Runtime/PacketHandlers/EncryptionComponents/AsymmetricEncryption/RSAEncryptionHandlerComponent/Public/RSAEncryptionHandlerComponent.h b/Engine/Source/Runtime/PacketHandlers/EncryptionComponents/AsymmetricEncryption/RSAEncryptionHandlerComponent/Public/RSAEncryptionHandlerComponent.h
--- a/Engine/Source/Runtime/PacketHandlers/EncryptionComponents/AsymmetricEncryption/RSAEncryptionHandlerComponent/Public/RSAEncryptionHandlerComponent.h
+++ b/Engine/Source/Runtime/PacketHandlers/EncryptionComponents/AsymmetricEncryption/RSAEncryptionHandlerComponent/Public/RSAEncryptionHandlerComponent.h
@@ -56,8 +56,13 @@ public:
 	/* Handles any outgoing packets */
 	virtual void Outgoing(FBitWriter& Packet, FOutPacketTraits& Traits) override;
 
-	virtual void IncomingConnectionless(const FString& Address, FBitReader& Packet) override {}
-	virtual void OutgoingConnectionless(const FString& Address, FBitWriter& Packet, FOutPacketTraits& Traits) override {}
+	virtual void IncomingConnectionless(const TSharedPtr<const FInternetAddr>& Address, FBitReader& Packet) override
+	{
+	}
+
+	virtual void OutgoingConnectionless(const TSharedPtr<const FInternetAddr>& Address, FBitWriter& Packet, FOutPacketTraits& Traits) override
+	{
+	}
 
 	virtual int32 GetReservedPacketBits() const override
 	{
diff --git a/Engine/Source/Runtime/PacketHandlers/EncryptionComponents/RSAKeyAESEncryption/Public/RSAKeyAESEncryptionHandlerComponent.h b/Engine/Source/Runtime/PacketHandlers/EncryptionComponents/RSAKeyAESEncryption/Public/RSAKeyAESEncryptionHandlerComponent.h
--- a/Engine/Source/Runtime/PacketHandlers/EncryptionComponents/RSAKeyAESEncryption/Public/RSAKeyAESEncryptionHandlerComponent.h
+++ b/Engine/Source/Runtime/PacketHandlers/EncryptionComponents/RSAKeyAESEncryption/Public/RSAKeyAESEncryptionHandlerComponent.h
@@ -60,8 +60,13 @@ public:
 	virtual void Incoming(FBitReader& Packet) override;
 	virtual void Outgoing(FBitWriter& Packet) override, FOutPacketTraits& Traits;
 
-	virtual void IncomingConnectionless(const FString& Address, FBitReader& Packet) override {}
-	virtual void OutgoingConnectionless(const FString& Address, FBitWriter& Packet, FOutPacketTraits& Traits) override {}
+	virtual void IncomingConnectionless(const TSharedPtr<const FInternetAddr>& Address, FBitReader& Packet) override
+	{
+	}
+
+	virtual void OutgoingConnectionless(const TSharedPtr<const FInternetAddr>& Address, FBitWriter& Packet, FOutPacketTraits& Traits) override
+	{
+	}
 
 	virtual int32 GetReservedPacketBits() const override;
 
diff --git a/Engine/Source/Runtime/PacketHandlers/PacketHandler/PacketHandler.Build.cs b/Engine/Source/Runtime/PacketHandlers/PacketHandler/PacketHandler.Build.cs
--- a/Engine/Source/Runtime/PacketHandlers/PacketHandler/PacketHandler.Build.cs
+++ b/Engine/Source/Runtime/PacketHandlers/PacketHandler/PacketHandler.Build.cs
@@ -14,6 +14,7 @@ public class PacketHandler : ModuleRules
 				"Core",
 				"CoreUObject",
                 "ReliabilityHandlerComponent",
+				"Sockets",
             }
         );
 
diff --git a/Engine/Source/Runtime/PacketHandlers/PacketHandler/Private/PacketHandler.cpp b/Engine/Source/Runtime/PacketHandlers/PacketHandler/Private/PacketHandler.cpp
--- a/Engine/Source/Runtime/PacketHandlers/PacketHandler/Private/PacketHandler.cpp
+++ b/Engine/Source/Runtime/PacketHandlers/PacketHandler/Private/PacketHandler.cpp
@@ -15,6 +15,8 @@
 #include "ReliabilityHandlerComponent.h"
 #include "PacketHandlerProfileConfig.h"
 
+#include "SocketSubsystem.h"
+
 // @todo #JohnB: There is quite a lot of inefficient copying of packet data going on.
 //					Redo the whole packet parsing/modification pipeline.
 
@@ -35,6 +37,7 @@ PacketHandler::PacketHandler(FDDoSDetection* InDDoS/*=nullptr*/)
 	, DDoS(InDDoS)
 	, LowLevelSendDel()
 	, LowLevelSendDel_Deprecated()
+	, NetDriverGetAddressFromString_Deprecated(nullptr)
 	, HandshakeCompleteDel()
 	, OutgoingPacket()
 	, IncomingPacket()
@@ -427,7 +430,8 @@ void HandlerComponent::CountBytes(FArchive& Ar) const
 	Ar.CountBytes(sizeof(*this), sizeof(*this));
 }
 
-const ProcessedPacket PacketHandler::Incoming_Internal(uint8* Packet, int32 CountBytes, bool bConnectionless, const FString& Address)
+
+const ProcessedPacket PacketHandler::Incoming_Internal(uint8* Packet, int32 CountBytes, bool bConnectionless, const TSharedPtr<const FInternetAddr>& Address)
 {
 	SCOPE_CYCLE_COUNTER(Stat_PacketHandler_Incoming_Internal);
 
@@ -520,7 +524,7 @@ const ProcessedPacket PacketHandler::Incoming_Internal(uint8* Packet, int32 Coun
 	}
 }
 
-const ProcessedPacket PacketHandler::Outgoing_Internal(uint8* Packet, int32 CountBits, FOutPacketTraits& Traits, bool bConnectionless, const FString& Address)
+const ProcessedPacket PacketHandler::Outgoing_Internal(uint8* Packet, int32 CountBits, FOutPacketTraits& Traits, bool bConnectionless, const TSharedPtr<const FInternetAddr>& Address)
 {
 	SCOPE_CYCLE_COUNTER(Stat_PacketHandler_Outgoing_Internal);
 
@@ -1008,3 +1012,32 @@ void FPacketHandlerComponentModuleInterface::ShutdownModule()
 	FPacketAudit::Destruct();
 }
 
+PRAGMA_DISABLE_DEPRECATION_WARNINGS
+TSharedPtr<const FInternetAddr> PacketHandler::GetAddressFromString(const FString& Address)
+{
+	if (NetDriverGetAddressFromString_Deprecated)
+	{
+		return NetDriverGetAddressFromString_Deprecated(Address);
+	}
+
+	return nullptr;
+}
+
+void BufferedPacket::SetAddressFromIP(const FString& InAddress)
+{
+	ISocketSubsystem* SocketSub = ISocketSubsystem::Get();
+	if (SocketSub != nullptr)
+	{
+		// TODO: @jleonard - Don't require the bracket removal here.
+		FString SanitizedAddress = InAddress;
+		SanitizedAddress.RemoveFromEnd(TEXT("]"));
+		SanitizedAddress.RemoveFromStart(TEXT("["));
+		TSharedPtr<FInternetAddr> NewAddress = SocketSub->GetAddressFromString(SanitizedAddress);
+		if (NewAddress.IsValid())
+		{
+			Address = NewAddress;
+		}
+	}
+}
+
+PRAGMA_ENABLE_DEPRECATION_WARNINGS
diff --git a/Engine/Source/Runtime/PacketHandlers/PacketHandler/Public/PacketHandler.h b/Engine/Source/Runtime/PacketHandlers/PacketHandler/Public/PacketHandler.h
--- a/Engine/Source/Runtime/PacketHandlers/PacketHandler/Public/PacketHandler.h
+++ b/Engine/Source/Runtime/PacketHandlers/PacketHandler/Public/PacketHandler.h
@@ -8,6 +8,7 @@
 #include "Modules/ModuleInterface.h"
 #include "Containers/Queue.h"
 #include "PacketTraits.h"
+#include "IPAddress.h"
 
 PACKETHANDLER_API DECLARE_LOG_CATEGORY_EXTERN(PacketHandlerLog, Log, All);
 
@@ -129,8 +130,8 @@ struct PACKETHANDLER_API BufferedPacket
 	/** Used by ReliabilityHandlerComponent, to track packet id's */
 	uint32 Id;
 
-	/** For connectionless packets, the address to send to (format is abstract, determined by active net driver) */
-	FString Address;
+	/** For connectionless packets, the address to send to */
+	TSharedPtr<const FInternetAddr> Address;
 
 	/** If buffering a packet through 'SendHandlerPacket', track the originating component */
 	HandlerComponent* FromComponent;
@@ -145,11 +146,13 @@ private:
 		, Traits()
 		, ResendTime(0.0)
 		, Id(0)
-		, Address()
+		, Address(nullptr)
 		, FromComponent(nullptr)
 	{
 	}
 
+	void SetAddressFromIP(const FString& InAddress);
+
 public:
 	UE_DEPRECATED(4.21, "Please use the new constructor that adds support for analytics and better precision")
 	BufferedPacket(uint8* InCopyData, uint32 InCountBits, float InResendTime=0.f, uint32 InId=0)
@@ -157,7 +160,7 @@ public:
 		, Traits()
 		, ResendTime(double(InResendTime))
 		, Id(InId)
-		, Address()
+		, Address(nullptr)
 		, FromComponent(nullptr)
 	{
 		check(InCopyData != nullptr);
@@ -171,7 +174,7 @@ public:
 		, Traits(InTraits)
 		, ResendTime(InResendTime)
 		, Id(InId)
-		, Address()
+		, Address(nullptr)
 		, FromComponent(nullptr)
 	{
 		check(InCopyData != nullptr);
@@ -180,24 +183,14 @@ public:
 		FMemory::Memcpy(Data, InCopyData, FMath::DivideAndRoundUp(InCountBits, 8u));
 	}
 
-	UE_DEPRECATED(4.21, "Please use the new constructor that adds support for analytics")
-	BufferedPacket(const FString& InAddress, uint8* InCopyData, uint32 InCountBits, double InResendTime=0.0, uint32 InId=0)
-		: CountBits(InCountBits)
-		, Traits()
-		, ResendTime(InResendTime)
-		, Id(InId)
-		, Address()
-		, FromComponent(nullptr)
+	UE_DEPRECATED(4.23, "Please use the new constructor takes in a FInternetAddr")
+	BufferedPacket(const FString& InAddress, uint8* InCopyData, uint32 InCountBits, FOutPacketTraits& InTraits, double InResendTime=0.0, uint32 InId=0)
+		: BufferedPacket(InCopyData, InCountBits, InTraits, double(InResendTime), InId)
 	{
-		check(InCopyData != nullptr);
-
-		Data = new uint8[FMath::DivideAndRoundUp(InCountBits, 8u)];
-		FMemory::Memcpy(Data, InCopyData, FMath::DivideAndRoundUp(InCountBits, 8u));
-
-		Address = InAddress;
+		SetAddressFromIP(InAddress);
 	}
 
-	BufferedPacket(const FString& InAddress, uint8* InCopyData, uint32 InCountBits, FOutPacketTraits& InTraits, double InResendTime=0.0, uint32 InId=0)
+	BufferedPacket(TSharedPtr<const FInternetAddr> InAddress, uint8* InCopyData, uint32 InCountBits, FOutPacketTraits& InTraits, double InResendTime=0.0, uint32 InId=0)
 		: BufferedPacket(InCopyData, InCountBits, InTraits, InResendTime, InId)
 	{
 		Address = InAddress;
@@ -215,7 +208,6 @@ public:
 	{
 		Ar.CountBytes(sizeof(*this), sizeof(*this));
 		Ar.CountBytes(FMath::DivideAndRoundUp(CountBits, 8u), FMath::DivideAndRoundUp(CountBits, 8u));
-		Address.CountBytes(Ar);
 	}
 };
 
@@ -261,6 +253,13 @@ public:
 		LowLevelSendDel = InLowLevelSendDel;
 	}
 
+	typedef TFunction<TSharedPtr<FInternetAddr>(const FString&)> NetDriverAddressFunc;
+	UE_DEPRECATED(4.23, "Do not use this function directly, it is provided for backwards compatibility with older versions of LowLevelSend")
+	void InitializeAddressSerializer(NetDriverAddressFunc InAddressSerializerFunc)
+	{
+		NetDriverGetAddressFromString_Deprecated = InAddressSerializerFunc;
+	}
+
 	/**
 	 * Notification that the NetDriver analytics provider has been updated (NOT called on first initialization)
 	 * NOTE: Can also mean disabled, e.g. during hotfix
@@ -342,8 +341,7 @@ public:
 	 */
 	FORCEINLINE const ProcessedPacket Incoming(uint8* Packet, int32 CountBytes)
 	{
-		static const FString EmptyString(TEXT(""));
-		return Incoming_Internal(Packet, CountBytes, false, EmptyString);
+		return Incoming_Internal(Packet, CountBytes, false, nullptr);
 	}
 
 	UE_DEPRECATED(4.21, "Please move to the functional flow that includes support for PacketTraits.")
@@ -366,8 +364,15 @@ public:
 	 */
 	FORCEINLINE const ProcessedPacket Outgoing(uint8* Packet, int32 CountBits, FOutPacketTraits& Traits)
 	{
-		static const FString EmptyString(TEXT(""));
-		return Outgoing_Internal(Packet, CountBits, Traits, false, EmptyString);
+		return Outgoing_Internal(Packet, CountBits, Traits, false, nullptr);
+	}
+
+	UE_DEPRECATED(4.22, "Use the FInternetAddr structure version.")
+	FORCEINLINE const ProcessedPacket IncomingConnectionless(const FString& Address, uint8* Packet, int32 CountBytes)
+	{
+		PRAGMA_DISABLE_DEPRECATION_WARNINGS
+		return IncomingConnectionless(GetAddressFromString(Address), Packet, CountBytes);
+		PRAGMA_ENABLE_DEPRECATION_WARNINGS
 	}
 
 	/**
@@ -375,12 +380,12 @@ public:
 	 * IMPORTANT: Net drivers triggering this, should call 'UNetDriver::FlushHandler' shortly afterwards, to minimize packet buffering
 	 * NOTE: Connectionless packets are unreliable.
 	 *
-	 * @param Address		The address the packet was received from (format is abstract, determined by active net driver)
+	 * @param Address		The address the packet was received from
 	 * @param Packet		The packet data to be processed
 	 * @param CountBytes	The size of the packet data in bytes
 	 * @return				Returns the final packet
 	 */
-	FORCEINLINE const ProcessedPacket IncomingConnectionless(const FString& Address, uint8* Packet, int32 CountBytes)
+	FORCEINLINE const ProcessedPacket IncomingConnectionless(const TSharedPtr<const FInternetAddr>& Address, uint8* Packet, int32 CountBytes)
 	{
 		return Incoming_Internal(Packet, CountBytes, true, Address);
 	}
@@ -389,20 +394,30 @@ public:
 	FORCEINLINE const ProcessedPacket OutgoingConnectionless(const FString& Address, uint8* Packet, int32 CountBits)
 	{
 		FOutPacketTraits EmptyTraits;
+		PRAGMA_DISABLE_DEPRECATION_WARNINGS
 		return OutgoingConnectionless(Address, Packet, CountBits, EmptyTraits);
+		PRAGMA_ENABLE_DEPRECATION_WARNINGS
+	}
+
+	UE_DEPRECATED(4.22, "Use the FInternetAddr structure version.")
+	FORCEINLINE const ProcessedPacket OutgoingConnectionless(const FString& Address, uint8* Packet, int32 CountBits, FOutPacketTraits& Traits)
+	{
+		PRAGMA_DISABLE_DEPRECATION_WARNINGS
+		return OutgoingConnectionless(GetAddressFromString(Address), Packet, CountBits, Traits);
+		PRAGMA_ENABLE_DEPRECATION_WARNINGS
 	}
 
 	/**
 	 * Processes outgoing packets without a UNetConnection, in the same manner as 'Outgoing' above
 	 * NOTE: Connectionless packets are unreliable.
 	 *
-	 * @param Address		The address the packet is being sent to (format is abstract, determined by active net driver)
+	 * @param Address		The address the packet is being sent to
 	 * @param Packet		The packet data to be processed
 	 * @param CountBits		The size of the packet data in bits
 	 * @param Traits		Traits for the packet, if applicable
 	 * @return				Returns the final packet
 	 */
-	FORCEINLINE const ProcessedPacket OutgoingConnectionless(const FString& Address, uint8* Packet, int32 CountBits, FOutPacketTraits& Traits)
+	FORCEINLINE const ProcessedPacket OutgoingConnectionless(const TSharedPtr<const FInternetAddr>& Address, uint8* Packet, int32 CountBits, FOutPacketTraits& Traits)
 	{
 		return Outgoing_Internal(Packet, CountBits, Traits, true, Address);
 	}
@@ -416,6 +431,7 @@ public:
 	virtual void CountBytes(FArchive& Ar) const;
 
 protected:
+
 	/**
 	 * Internal handling for Incoming/IncomingConnectionless
 	 *
@@ -425,7 +441,24 @@ protected:
 	 * @param Address			The address the packet was received from (format is abstract, determined by active net driver)
 	 * @return					Returns the final packet
 	 */
-	const ProcessedPacket Incoming_Internal(uint8* Packet, int32 CountBytes, bool bConnectionless, const FString& Address);
+	UE_DEPRECATED(4.23, "Address is now taken as a shared pointer to remove the burden of format differences from the packethandler")
+	const ProcessedPacket Incoming_Internal(uint8* Packet, int32 CountBytes, bool bConnectionless, const FString& Address)
+	{
+		PRAGMA_DISABLE_DEPRECATION_WARNINGS
+		return Incoming_Internal(Packet, CountBytes, bConnectionless, GetAddressFromString(Address));
+		PRAGMA_ENABLE_DEPRECATION_WARNINGS
+	}
+
+	/**
+	 * Internal handling for Incoming/IncomingConnectionless
+	 *
+	 * @param Packet			The packet data to be processed
+	 * @param CountBytes		The size of the packet data in bytes
+	 * @param bConnectionless	Whether or not this should be processed as a connectionless packet
+	 * @param Address			The address the packet was received from
+	 * @return					Returns the final packet
+	 */
+	const ProcessedPacket Incoming_Internal(uint8* Packet, int32 CountBytes, bool bConnectionless, const TSharedPtr<const FInternetAddr>& Address);
 
 	/**
 	 * Internal handling for Outgoing/OutgoingConnectionless
@@ -437,10 +470,34 @@ protected:
 	 * @param Address			The address the packet is being sent to (format is abstract, determined by active net driver)
 	 * @return					Returns the final packet
 	 */
-	const ProcessedPacket Outgoing_Internal(uint8* Packet, int32 CountBits, FOutPacketTraits& Traits, bool bConnectionless, const FString& Address );
+	UE_DEPRECATED(4.23, "Address is now taken as a shared pointer to remove the burden of format differences from the packethandler")
+	const ProcessedPacket Outgoing_Internal(uint8* Packet, int32 CountBits, FOutPacketTraits& Traits, bool bConnectionless, const FString& Address)
+	{
+		PRAGMA_DISABLE_DEPRECATION_WARNINGS
+		return Outgoing_Internal(Packet, CountBits, Traits, bConnectionless, GetAddressFromString(Address));
+		PRAGMA_ENABLE_DEPRECATION_WARNINGS
+	}
+
+	/**
+	 * Internal handling for Outgoing/OutgoingConnectionless
+	 *
+	 * @param Packet			The packet data to be processed
+	 * @param CountBits			The size of the packet data in bits
+	 * @param Traits			Traits for the packet, passed down from the NetConnection, if applicable
+	 * @param bConnectionless	Whether or not this should be sent as a connectionless packet
+	 * @param Address			The address the packet is being sent to
+	 * @return					Returns the final packet
+	 */
+	const ProcessedPacket Outgoing_Internal(uint8* Packet, int32 CountBits, FOutPacketTraits& Traits, bool bConnectionless, const TSharedPtr<const FInternetAddr>& Address);
 
 public:
 
+	/*
+	 * Allows for easy translation from string to FInternetAddr. Only used for compatibility in deprecated functions, do not use in new code flow.
+	 */
+	UE_DEPRECATED(4.23, "This function should not be used except for backwards compatibility with deprecated functions.")
+	TSharedPtr<const FInternetAddr> GetAddressFromString(const FString& Address);
+
 	UE_DEPRECATED(4.21, "Please use the packet traits when sending to handle modifications of packets and analytics.")
 	void SendHandlerPacket(HandlerComponent* InComponent, FBitWriter& Writer)
 	{
@@ -618,6 +675,9 @@ private:
 	/** Delegate used for triggering PacketHandler/HandlerComponent-sourced sends (DEPRECATED) */
 	FPacketHandlerLowLevelSend LowLevelSendDel_Deprecated;
 
+	/** Function for the netdriver to provide serialization of FStrings in deprecated use cases (DEPRECATED) */
+	NetDriverAddressFunc NetDriverGetAddressFromString_Deprecated;
+
 	/** Delegate used for notifying that handshaking has completed */
 	FPacketHandlerHandshakeComplete HandshakeCompleteDel;
 
@@ -751,29 +811,53 @@ public:
 	 */
 	virtual void Outgoing(FBitWriter& Packet, FOutPacketTraits& Traits) = 0;
 
+	UE_DEPRECATED(4.23, "Use a FInternetAddr instead of a string for more versitile behavior")
+	virtual void IncomingConnectionless(const FString& Address, FBitReader& Packet)
+	{
+		if (Handler)
+		{
+			PRAGMA_DISABLE_DEPRECATION_WARNINGS
+			IncomingConnectionless(Handler->GetAddressFromString(Address), Packet);
+			PRAGMA_ENABLE_DEPRECATION_WARNINGS
+		}
+	}
+
 	/**
 	 * Handles incoming packets not associated with a UNetConnection
 	 *
-	 * @param Address	The address the packet was received from (format is abstract, determined by active net driver)
+	 * @param Address	The address the packet was received from
 	 * @param Packet	The packet to be handled
 	 */
-	virtual void IncomingConnectionless(const FString& Address, FBitReader& Packet) = 0;
+	virtual void IncomingConnectionless(const TSharedPtr<const FInternetAddr>& Address, FBitReader& Packet) = 0;
 
 	UE_DEPRECATED(4.21, "Use the method that allows traits on the packet.")
 	virtual void OutgoingConnectionless(const FString& Address, FBitWriter& Packet)
 	{
 		FOutPacketTraits EmptyTraits;
+		PRAGMA_DISABLE_DEPRECATION_WARNINGS
 		OutgoingConnectionless(Address, Packet, EmptyTraits);
+		PRAGMA_ENABLE_DEPRECATION_WARNINGS
+	}
+
+	UE_DEPRECATED(4.23, "Use a FInternetAddr instead of a string for more versitile behavior")
+	virtual void OutgoingConnectionless(const FString& Address, FBitWriter& Packet, FOutPacketTraits& Traits)
+	{
+		if (Handler)
+		{
+			PRAGMA_DISABLE_DEPRECATION_WARNINGS
+			OutgoingConnectionless(Handler->GetAddressFromString(Address), Packet, Traits);
+			PRAGMA_ENABLE_DEPRECATION_WARNINGS
+		}
 	}
 
 	/**
 	 * Handles any outgoing packets not associated with a UNetConnection
 	 *
-	 * @param Address	The address the packet is being sent to (format is abstract, determined by active net driver)
+	 * @param Address	The address the packet is being sent to
 	 * @param Packet	The packet to be handled
 	 * @param Traits	Traits for the packet, passed down through the packet pipeline (if applicable)
 	 */
-	virtual void OutgoingConnectionless(const FString& Address, FBitWriter& Packet, FOutPacketTraits& Traits) = 0;
+	virtual void OutgoingConnectionless(const TSharedPtr<const FInternetAddr>& Address, FBitWriter& Packet, FOutPacketTraits& Traits) = 0;
 
 
 	/**
diff --git a/Engine/Source/Runtime/PacketHandlers/ReliabilityHandlerComponent/Public/ReliabilityHandlerComponent.h b/Engine/Source/Runtime/PacketHandlers/ReliabilityHandlerComponent/Public/ReliabilityHandlerComponent.h
--- a/Engine/Source/Runtime/PacketHandlers/ReliabilityHandlerComponent/Public/ReliabilityHandlerComponent.h
+++ b/Engine/Source/Runtime/PacketHandlers/ReliabilityHandlerComponent/Public/ReliabilityHandlerComponent.h
@@ -24,8 +24,8 @@ public:
 	virtual void Incoming(FBitReader& Packet) override;
 	virtual void Outgoing(FBitWriter& Packet, FOutPacketTraits& Traits) override;
 
-	virtual void IncomingConnectionless(const FString& Address, FBitReader& Packet) override {}
-	virtual void OutgoingConnectionless(const FString& Address, FBitWriter& Packet, FOutPacketTraits& Traits) override {}
+	virtual void IncomingConnectionless(const TSharedPtr<const FInternetAddr>& Address, FBitReader& Packet) override {}
+	virtual void OutgoingConnectionless(const TSharedPtr<const FInternetAddr>& Address, FBitWriter& Packet, FOutPacketTraits& Traits) override {}
 
 	/* Queues a packet for resending */
 	void QueuePacketForResending(uint8* Packet, int32 CountBits, FOutPacketTraits& Traits);
diff --git a/Engine/Source/Runtime/Sockets/Public/IPAddress.h b/Engine/Source/Runtime/Sockets/Public/IPAddress.h
--- a/Engine/Source/Runtime/Sockets/Public/IPAddress.h
+++ b/Engine/Source/Runtime/Sockets/Public/IPAddress.h
@@ -400,3 +400,21 @@ struct FInternetAddrKeyMapFuncs : public BaseKeyFuncs<ValueType, TSharedRef<FInt
 	}
 };
 
+template<typename ValueType>
+struct FInternetAddrConstKeyMapFuncs : public BaseKeyFuncs<ValueType, TSharedRef<const FInternetAddr>, false>
+{
+	static FORCEINLINE const TSharedRef<const FInternetAddr>& GetSetKey(const TPair<TSharedRef<const FInternetAddr>, ValueType>& Element)
+	{
+		return Element.Key;
+	}
+
+	static FORCEINLINE bool Matches(const TSharedRef<const FInternetAddr>& A, const TSharedRef<const FInternetAddr>& B)
+	{
+		return *A == *B;
+	}
+
+	static FORCEINLINE uint32 GetKeyHash(const TSharedRef<const FInternetAddr>& Key)
+	{
+		return Key->GetTypeHash();
+	}
+};
