Fix iOS and Mac build break due to compiler ambiguity.

#rnx

diff --git a/Engine/Source/Runtime/PacketHandlers/PacketHandler/Public/PacketHandler.h b/Engine/Source/Runtime/PacketHandlers/PacketHandler/Public/PacketHandler.h
--- a/Engine/Source/Runtime/PacketHandlers/PacketHandler/Public/PacketHandler.h
+++ b/Engine/Source/Runtime/PacketHandlers/PacketHandler/Public/PacketHandler.h
@@ -341,7 +341,8 @@ public:
 	 */
 	FORCEINLINE const ProcessedPacket Incoming(uint8* Packet, int32 CountBytes)
 	{
-		return Incoming_Internal(Packet, CountBytes, false, nullptr);
+		const TSharedPtr<const FInternetAddr> EmptyAddress = nullptr;
+		return Incoming_Internal(Packet, CountBytes, false, EmptyAddress);
 	}
 
 	UE_DEPRECATED(4.21, "Please move to the functional flow that includes support for PacketTraits.")
@@ -364,7 +365,8 @@ public:
 	 */
 	FORCEINLINE const ProcessedPacket Outgoing(uint8* Packet, int32 CountBits, FOutPacketTraits& Traits)
 	{
-		return Outgoing_Internal(Packet, CountBits, Traits, false, nullptr);
+		const TSharedPtr<const FInternetAddr> EmptyAddress = nullptr;
+		return Outgoing_Internal(Packet, CountBits, Traits, false, EmptyAddress);
 	}
 
 	UE_DEPRECATED(4.22, "Use the FInternetAddr structure version.")
