Added detection for failed (debug-disabled) HandlerComponent loads.

diff --git a/Engine/Source/Runtime/PacketHandlers/PacketHandler/Private/PacketHandler.cpp b/Engine/Source/Runtime/PacketHandlers/PacketHandler/Private/PacketHandler.cpp
--- a/Engine/Source/Runtime/PacketHandlers/PacketHandler/Private/PacketHandler.cpp
+++ b/Engine/Source/Runtime/PacketHandlers/PacketHandler/Private/PacketHandler.cpp
@@ -322,11 +322,12 @@ TSharedPtr<HandlerComponent> PacketHandler::AddHandler(const FString& ComponentS
 			{
 				FPacketHandlerComponentModuleInterface* PacketHandlerInterface = FModuleManager::Get().LoadModulePtr<FPacketHandlerComponentModuleInterface>(FName(*ComponentName));
 
-				if (PacketHandlerInterface)
+				if (PacketHandlerInterface != nullptr)
 				{
 					ReturnVal = PacketHandlerInterface->CreateComponentInstance(ComponentOptions);
 				}
-				else
+
+				if (!ReturnVal.IsValid())
 				{
 					UE_LOG(PacketHandlerLog, Warning, TEXT("Unable to Load Module: %s"), *ComponentName);
 				}
