Deprecate CryptoPP integration and the old encryption packet handlers that depend on it

diff --git a/Engine/Source/Runtime/PacketHandlers/EncryptionComponents/AsymmetricEncryption/RSAEncryptionHandlerComponent/Public/RSAEncryptionHandlerComponent.h b/Engine/Source/Runtime/PacketHandlers/EncryptionComponents/AsymmetricEncryption/RSAEncryptionHandlerComponent/Public/RSAEncryptionHandlerComponent.h
--- a/Engine/Source/Runtime/PacketHandlers/EncryptionComponents/AsymmetricEncryption/RSAEncryptionHandlerComponent/Public/RSAEncryptionHandlerComponent.h
+++ b/Engine/Source/Runtime/PacketHandlers/EncryptionComponents/AsymmetricEncryption/RSAEncryptionHandlerComponent/Public/RSAEncryptionHandlerComponent.h
@@ -10,6 +10,7 @@ THIRD_PARTY_INCLUDES_START
 THIRD_PARTY_INCLUDES_END
 
 /* RSA Encryptor Module Interface */
+UE_DEPRECATED(5.0, "This Handler Is Now Deprecated");
 class FRSAEncryptorHandlerComponentModuleInterface : public FPacketHandlerComponentModuleInterface
 {
 public:
@@ -31,6 +32,7 @@ namespace ERSAEncryptionHandler
 /*
 * Asymmetric block cipher handler component using RSA as the cipher
 */
+UE_DEPRECATED(5.0, "This Handler Is Now Deprecated");
 class RSAENCRYPTIONHANDLERCOMPONENT_API RSAEncryptionHandlerComponent : public HandlerComponent
 {
 public:
diff --git a/Engine/Source/Runtime/PacketHandlers/EncryptionComponents/RSAKeyAESEncryption/Public/RSAKeyAESEncryptionHandlerComponent.h b/Engine/Source/Runtime/PacketHandlers/EncryptionComponents/RSAKeyAESEncryption/Public/RSAKeyAESEncryptionHandlerComponent.h
--- a/Engine/Source/Runtime/PacketHandlers/EncryptionComponents/RSAKeyAESEncryption/Public/RSAKeyAESEncryptionHandlerComponent.h
+++ b/Engine/Source/Runtime/PacketHandlers/EncryptionComponents/RSAKeyAESEncryption/Public/RSAKeyAESEncryptionHandlerComponent.h
@@ -19,6 +19,7 @@ THIRD_PARTY_INCLUDES_END
 /**
  * Module Interface
  */
+UE_DEPRECATED(5.0, "This Handler Is Now Deprecated");
 class FRSAKeyAESEncryptionModuleInterface : public FPacketHandlerComponentModuleInterface
 {
 public:
@@ -39,6 +40,7 @@ namespace ERSAKeyAESEncryptionHandler
 /**
  * Exchanges a key using RSA as the handshake, and implements symmetric encryption with that key using AES
  */
+UE_DEPRECATED(5.0, "This Handler Is Now Deprecated");
 class RSAKEYAESENCRYPTION_API RSAKeyAESEncryptionHandlerComponent : public HandlerComponent
 {
 public:
diff --git a/Engine/Source/Runtime/PacketHandlers/EncryptionComponents/SymmetricEncryption/BlockEncryption/AESBlockEncryptor/Public/AESBlockEncryptor.h b/Engine/Source/Runtime/PacketHandlers/EncryptionComponents/SymmetricEncryption/BlockEncryption/AESBlockEncryptor/Public/AESBlockEncryptor.h
--- a/Engine/Source/Runtime/PacketHandlers/EncryptionComponents/SymmetricEncryption/BlockEncryption/AESBlockEncryptor/Public/AESBlockEncryptor.h
+++ b/Engine/Source/Runtime/PacketHandlers/EncryptionComponents/SymmetricEncryption/BlockEncryption/AESBlockEncryptor/Public/AESBlockEncryptor.h
@@ -6,6 +6,7 @@
 #include "CryptoPP/5.6.5/include/aes.h"
 
 /* AES Block Encryptor Module Interface */
+UE_DEPRECATED(5.0, "This Encryptor Is Now Deprecated");
 class FAESBlockEncryptorModuleInterface : public FBlockEncryptorModuleInterface
 {
 	virtual BlockEncryptor* CreateBlockEncryptorInstance() override;
@@ -14,6 +15,7 @@ class FAESBlockEncryptorModuleInterface : public FBlockEncryptorModuleInterface
 /*
 * AES Block encryption
 */
+UE_DEPRECATED(5.0, "This Encryptor Is Now Deprecated");
 class AESBLOCKENCRYPTOR_API AESBlockEncryptor : public BlockEncryptor
 {
 public:
diff --git a/Engine/Source/Runtime/PacketHandlers/EncryptionComponents/SymmetricEncryption/BlockEncryption/BlockEncryptionHandlerComponent/Public/BlockEncryptionHandlerComponent.h b/Engine/Source/Runtime/PacketHandlers/EncryptionComponents/SymmetricEncryption/BlockEncryption/BlockEncryptionHandlerComponent/Public/BlockEncryptionHandlerComponent.h
--- a/Engine/Source/Runtime/PacketHandlers/EncryptionComponents/SymmetricEncryption/BlockEncryption/BlockEncryptionHandlerComponent/Public/BlockEncryptionHandlerComponent.h
+++ b/Engine/Source/Runtime/PacketHandlers/EncryptionComponents/SymmetricEncryption/BlockEncryption/BlockEncryptionHandlerComponent/Public/BlockEncryptionHandlerComponent.h
@@ -9,6 +9,7 @@
 /*
 * Abstract class for block encryptors
 */
+UE_DEPRECATED(5.0, "This Encryptor Is Now Deprecated");
 class BLOCKENCRYPTIONHANDLERCOMPONENT_API BlockEncryptor
 {
 public:
@@ -36,6 +37,7 @@ protected:
 };
 
 /* Block Encryptor Module Interface */
+UE_DEPRECATED(5.0, "This Encryptor Is Now Deprecated");
 class FBlockEncryptorModuleInterface : public IModuleInterface
 {
 public:
@@ -45,6 +47,7 @@ public:
 /*
 * Symmetric block cipher handler component using AES as the cipher
 */
+UE_DEPRECATED(5.0, "This Encryptor Is Now Deprecated");
 class BLOCKENCRYPTIONHANDLERCOMPONENT_API BlockEncryptionHandlerComponent : public HandlerComponent
 {
 public:
@@ -80,6 +83,7 @@ protected:
 };
 
 /* Block Encryption Handler Component Module Interface */
+UE_DEPRECATED(5.0, "This Encryptor Is Now Deprecated");
 class FBlockEncryptionHandlerComponentModuleInterface : public FPacketHandlerComponentModuleInterface
 {
 public:
diff --git a/Engine/Source/Runtime/PacketHandlers/EncryptionComponents/SymmetricEncryption/BlockEncryption/BlowFishBlockEncryptor/Public/BlowFishBlockEncryptor.h b/Engine/Source/Runtime/PacketHandlers/EncryptionComponents/SymmetricEncryption/BlockEncryption/BlowFishBlockEncryptor/Public/BlowFishBlockEncryptor.h
--- a/Engine/Source/Runtime/PacketHandlers/EncryptionComponents/SymmetricEncryption/BlockEncryption/BlowFishBlockEncryptor/Public/BlowFishBlockEncryptor.h
+++ b/Engine/Source/Runtime/PacketHandlers/EncryptionComponents/SymmetricEncryption/BlockEncryption/BlowFishBlockEncryptor/Public/BlowFishBlockEncryptor.h
@@ -6,6 +6,7 @@
 #include "CryptoPP/5.6.5/include/blowfish.h"
 
 /* BlowFish Block Encryptor Module Interface */
+UE_DEPRECATED(5.0, "This Encryptor Is Now Deprecated");
 class FBlowFishBlockEncryptorModuleInterface : public FBlockEncryptorModuleInterface
 {
 	virtual BlockEncryptor* CreateBlockEncryptorInstance() override;
@@ -15,6 +16,7 @@ class FBlowFishBlockEncryptorModuleInterface : public FBlockEncryptorModuleInter
 /*
 * BlowFish Block encryption
 */
+UE_DEPRECATED(5.0, "This Encryptor Is Now Deprecated");
 class BLOWFISHBLOCKENCRYPTOR_API BlowFishBlockEncryptor : public BlockEncryptor
 {
 public:
diff --git a/Engine/Source/Runtime/PacketHandlers/EncryptionComponents/SymmetricEncryption/BlockEncryption/TwoFishBlockEncryptor/Public/TwoFishBlockEncryptor.h b/Engine/Source/Runtime/PacketHandlers/EncryptionComponents/SymmetricEncryption/BlockEncryption/TwoFishBlockEncryptor/Public/TwoFishBlockEncryptor.h
--- a/Engine/Source/Runtime/PacketHandlers/EncryptionComponents/SymmetricEncryption/BlockEncryption/TwoFishBlockEncryptor/Public/TwoFishBlockEncryptor.h
+++ b/Engine/Source/Runtime/PacketHandlers/EncryptionComponents/SymmetricEncryption/BlockEncryption/TwoFishBlockEncryptor/Public/TwoFishBlockEncryptor.h
@@ -6,6 +6,7 @@
 #include "CryptoPP/5.6.5/include/twofish.h"
 
 /* TwoFish Block Encryptor Module Interface */
+UE_DEPRECATED(5.0, "This Encryptor Is Now Deprecated");
 class FTwoFishBlockEncryptorModuleInterface : public FBlockEncryptorModuleInterface
 {
 	virtual BlockEncryptor* CreateBlockEncryptorInstance() override;
@@ -14,6 +15,7 @@ class FTwoFishBlockEncryptorModuleInterface : public FBlockEncryptorModuleInterf
 /*
 * TwoFish Block encryption
 */
+UE_DEPRECATED(5.0, "This Encryptor Is Now Deprecated");
 class TWOFISHBLOCKENCRYPTOR_API TwoFishBlockEncryptor : public BlockEncryptor
 {
 public:
diff --git a/Engine/Source/Runtime/PacketHandlers/EncryptionComponents/SymmetricEncryption/BlockEncryption/XORBlockEncryptor/Public/XORBlockEncryptor.h b/Engine/Source/Runtime/PacketHandlers/EncryptionComponents/SymmetricEncryption/BlockEncryption/XORBlockEncryptor/Public/XORBlockEncryptor.h
--- a/Engine/Source/Runtime/PacketHandlers/EncryptionComponents/SymmetricEncryption/BlockEncryption/XORBlockEncryptor/Public/XORBlockEncryptor.h
+++ b/Engine/Source/Runtime/PacketHandlers/EncryptionComponents/SymmetricEncryption/BlockEncryption/XORBlockEncryptor/Public/XORBlockEncryptor.h
@@ -5,6 +5,7 @@
 #include "BlockEncryptionHandlerComponent.h"
 
 /* XOR Block Encryptor Module Interface */
+UE_DEPRECATED(5.0, "This Encryptor Is Now Deprecated");
 class FXORBlockEncryptorModuleInterface : public FBlockEncryptorModuleInterface
 {
 	virtual BlockEncryptor* CreateBlockEncryptorInstance() override;
@@ -13,6 +14,7 @@ class FXORBlockEncryptorModuleInterface : public FBlockEncryptorModuleInterface
 /*
 * XOR Block encryption
 */
+UE_DEPRECATED(5.0, "This Encryptor Is Now Deprecated");
 class XORBLOCKENCRYPTOR_API XORBlockEncryptor : public BlockEncryptor
 {
 public:
diff --git a/Engine/Source/Runtime/PacketHandlers/EncryptionComponents/SymmetricEncryption/StreamEncryption/StreamEncryptionHandlerComponent/Public/StreamEncryptionHandlerComponent.h b/Engine/Source/Runtime/PacketHandlers/EncryptionComponents/SymmetricEncryption/StreamEncryption/StreamEncryptionHandlerComponent/Public/StreamEncryptionHandlerComponent.h
--- a/Engine/Source/Runtime/PacketHandlers/EncryptionComponents/SymmetricEncryption/StreamEncryption/StreamEncryptionHandlerComponent/Public/StreamEncryptionHandlerComponent.h
+++ b/Engine/Source/Runtime/PacketHandlers/EncryptionComponents/SymmetricEncryption/StreamEncryption/StreamEncryptionHandlerComponent/Public/StreamEncryptionHandlerComponent.h
@@ -4,6 +4,7 @@
 
 #include "PacketHandler.h"
 
+UE_DEPRECATED(5.0, "This Encryptor Is Now Deprecated");
 class STREAMENCRYPTIONHANDLERCOMPONENT_API StreamEncryptor
 {
 public:
@@ -25,6 +26,7 @@ protected:
 };
 
 /* Stream Encryptor Module Interface */
+UE_DEPRECATED(5.0, "This Encryptor Is Now Deprecated");
 class FStreamEncryptorModuleInterface : public IModuleInterface
 {
 public:
@@ -34,6 +36,7 @@ public:
 /*
 * Symmetric block cipher handler component using AES as the cipher
 */
+UE_DEPRECATED(5.0, "This Encryptor Is Now Deprecated");
 class STREAMENCRYPTIONHANDLERCOMPONENT_API StreamEncryptionHandlerComponent : public HandlerComponent
 {
 public:
@@ -69,6 +72,7 @@ protected:
 };
 
 /* Stream Encryption Handler Component Module Interface */
+UE_DEPRECATED(5.0, "This Encryptor Is Now Deprecated");
 class FStreamEncryptionHandlerComponentModuleInterface : public FPacketHandlerComponentModuleInterface
 {
 public:
diff --git a/Engine/Source/Runtime/PacketHandlers/EncryptionComponents/SymmetricEncryption/StreamEncryption/XORStreamEncryptor/Public/XORStreamEncryptor.h b/Engine/Source/Runtime/PacketHandlers/EncryptionComponents/SymmetricEncryption/StreamEncryption/XORStreamEncryptor/Public/XORStreamEncryptor.h
--- a/Engine/Source/Runtime/PacketHandlers/EncryptionComponents/SymmetricEncryption/StreamEncryption/XORStreamEncryptor/Public/XORStreamEncryptor.h
+++ b/Engine/Source/Runtime/PacketHandlers/EncryptionComponents/SymmetricEncryption/StreamEncryption/XORStreamEncryptor/Public/XORStreamEncryptor.h
@@ -3,6 +3,7 @@
 #include "StreamEncryptionHandlerComponent.h"
 
 /* XOR Stream Encryptor Module Interface */
+UE_DEPRECATED(5.0, "This Encryptor Is Now Deprecated");
 class FXORStreamEncryptorModuleInterface : public FStreamEncryptorModuleInterface
 {
 public:
@@ -12,6 +13,7 @@ public:
 /*
 * XOR Block encryption
 */
+UE_DEPRECATED(5.0, "This Encryptor Is Now Deprecated");
 class XORSTREAMENCRYPTOR_API XORStreamEncryptor : public StreamEncryptor
 {
 public:
