Add debug logging to help narrow down PacketHandler hitches

If you start with something like -dpcvars="net.PacketHandlerTimeguardThresholdMS=2" in non-shipping builds, you should see warning messages like:

$ grep PacketHandler *.log | grep -i warning
  [[REDACTED_PHONE]:018][119]PacketHandlerLog: Warning: PacketHandler: Incoming - OodleNetworkHandlerComponent took 62.68ms!
  [[REDACTED_PHONE]:732][150]PacketHandlerLog: Warning: PacketHandler: Incoming - OodleNetworkHandlerComponent took 63.56ms!

First two hitches at startup are expected, everything after that... not so much.

diff --git a/Engine/Source/Runtime/PacketHandlers/PacketHandler/Private/PacketHandler.cpp b/Engine/Source/Runtime/PacketHandlers/PacketHandler/Private/PacketHandler.cpp
--- a/Engine/Source/Runtime/PacketHandlers/PacketHandler/Private/PacketHandler.cpp
+++ b/Engine/Source/Runtime/PacketHandlers/PacketHandler/Private/PacketHandler.cpp
@@ -37,6 +37,41 @@ FAutoConsoleVariableRef CVarNetPacketHandlerCRCDump(
 	TEXT("net.PacketHandlerCRCDump"),
 	GPacketHandlerCRCDump,
 	TEXT("Enables or disables dumping of packet CRC's for every HandlerComponent, Incoming and Outgoing, for debugging."));
+
+static int32 GPacketHandlerTimeguardLimit = 20;
+static float GPacketHandlerTimeguardThresholdMS = 0.0f;
+
+static FAutoConsoleVariableRef CVarNetPacketHandlerTimeguardThresholdMS(
+	TEXT("net.PacketHandlerTimeguardThresholdMS"),
+	GPacketHandlerTimeguardThresholdMS,
+	TEXT("Threshold in milliseconds for the HandlerComponent timeguard, Incoming and Outgoing."),
+	ECVF_Default);
+static FAutoConsoleVariableRef CVarNetPacketHandlerTimeguardLimit(
+	TEXT("net.PacketHandlerTimeguardLimit"),
+	GPacketHandlerTimeguardLimit,
+	TEXT("Sets the maximum number of HandlerComponent timeguard logs.\n"),
+	ECVF_Default
+);
+
+// Lightweight time guard. Note: Threshold of 0 disables the timeguard
+#define NET_LIGHTWEIGHT_TIME_GUARD_BEGIN( Name, ThresholdMS ) \
+	float PREPROCESSOR_JOIN(__TimeGuard_ThresholdMS_, Name) = ThresholdMS; \
+	uint64 PREPROCESSOR_JOIN(__TimeGuard_StartCycles_, Name) = ( ThresholdMS > 0.0f && GPacketHandlerTimeguardLimit > 0 ) ? FPlatformTime::Cycles64() : 0;
+
+#define NET_LIGHTWEIGHT_TIME_GUARD_END( Name, NameStringCode ) \
+	if ( PREPROCESSOR_JOIN(__TimeGuard_ThresholdMS_, Name) > 0.0f && GPacketHandlerTimeguardLimit > 0 ) \
+	{\
+		float PREPROCESSOR_JOIN(__TimeGuard_MSElapsed_,Name) = FPlatformTime::ToMilliseconds64( FPlatformTime::Cycles64() - PREPROCESSOR_JOIN(__TimeGuard_StartCycles_,Name) ); \
+		if ( PREPROCESSOR_JOIN(__TimeGuard_MSElapsed_,Name) > PREPROCESSOR_JOIN(__TimeGuard_ThresholdMS_, Name) ) \
+		{ \
+			FString ReportName = NameStringCode; \
+			UE_LOG(PacketHandlerLog, Warning, TEXT("PacketHandler: %s - %s took %.2fms!"), TEXT(#Name), *ReportName, PREPROCESSOR_JOIN(__TimeGuard_MSElapsed_,Name)); \
+			GPacketHandlerTimeguardLimit--; \
+		} \
+	}
+#else // UE_BUILD_SHIPPING
+  #define NET_LIGHTWEIGHT_TIME_GUARD_BEGIN( Name, ThresholdMS )
+  #define NET_LIGHTWEIGHT_TIME_GUARD_END( Name, NameStringCode )
 #endif
 
 
@@ -593,7 +628,11 @@ EIncomingResult PacketHandler::Incoming_Internal(FReceivedPacketView& PacketView
 				}
 				else
 				{
+					NET_LIGHTWEIGHT_TIME_GUARD_BEGIN(Incoming, GPacketHandlerTimeguardThresholdMS);
+
 					CurComponent.Incoming(PacketRef);
+
+					NET_LIGHTWEIGHT_TIME_GUARD_END(Incoming, CurComponent.GetName().ToString());
 				}
 			}
 		}
@@ -710,7 +749,11 @@ const ProcessedPacket PacketHandler::Outgoing_Internal(uint8* Packet, int32 Coun
 						}
 						else
 						{
+							NET_LIGHTWEIGHT_TIME_GUARD_BEGIN(Outgoing, GPacketHandlerTimeguardThresholdMS);
+
 							CurComponent.Outgoing(OutgoingPacket, Traits);
+
+							NET_LIGHTWEIGHT_TIME_GUARD_END(Outgoing, CurComponent.GetName().ToString());
 						}
 					}
 					else
@@ -1213,4 +1256,4 @@ void FPacketHandlerComponentModuleInterface::StartupModule()
 void FPacketHandlerComponentModuleInterface::ShutdownModule()
 {
 	FPacketAudit::Destruct();
-}
\ No newline at end of file
+}
