Properly deprecate obsolete packet handler components.

diff --git a/Engine/Source/Runtime/PacketHandlers/EncryptionComponents/AsymmetricEncryption/RSAEncryptionHandlerComponent/Public/RSAEncryptionHandlerComponent.h b/Engine/Source/Runtime/PacketHandlers/EncryptionComponents/AsymmetricEncryption/RSAEncryptionHandlerComponent/Public/RSAEncryptionHandlerComponent.h
--- a/Engine/Source/Runtime/PacketHandlers/EncryptionComponents/AsymmetricEncryption/RSAEncryptionHandlerComponent/Public/RSAEncryptionHandlerComponent.h
+++ b/Engine/Source/Runtime/PacketHandlers/EncryptionComponents/AsymmetricEncryption/RSAEncryptionHandlerComponent/Public/RSAEncryptionHandlerComponent.h
@@ -10,8 +10,8 @@ THIRD_PARTY_INCLUDES_START
 THIRD_PARTY_INCLUDES_END
 
 /* RSA Encryptor Module Interface */
-UE_DEPRECATED(5.0, "This Handler Is Now Deprecated");
-class FRSAEncryptorHandlerComponentModuleInterface : public FPacketHandlerComponentModuleInterface
+class UE_DEPRECATED(5.3, "This component is not supported for encryption.")
+FRSAEncryptorHandlerComponentModuleInterface : public FPacketHandlerComponentModuleInterface
 {
 public:
 	/* Creates an instance of this component */
@@ -32,8 +32,8 @@ namespace ERSAEncryptionHandler
 /*
 * Asymmetric block cipher handler component using RSA as the cipher
 */
-UE_DEPRECATED(5.0, "This Handler Is Now Deprecated");
-class RSAENCRYPTIONHANDLERCOMPONENT_API RSAEncryptionHandlerComponent : public HandlerComponent
+class UE_DEPRECATED(5.3, "This component is not supported for encryption.")
+RSAENCRYPTIONHANDLERCOMPONENT_API RSAEncryptionHandlerComponent : public HandlerComponent
 {
 public:
 	/*
diff --git a/Engine/Source/Runtime/PacketHandlers/EncryptionComponents/EncryptionHandlerComponent/Public/EncryptionHandlerComponent.h b/Engine/Source/Runtime/PacketHandlers/EncryptionComponents/EncryptionHandlerComponent/Public/EncryptionHandlerComponent.h
--- a/Engine/Source/Runtime/PacketHandlers/EncryptionComponents/EncryptionHandlerComponent/Public/EncryptionHandlerComponent.h
+++ b/Engine/Source/Runtime/PacketHandlers/EncryptionComponents/EncryptionHandlerComponent/Public/EncryptionHandlerComponent.h
@@ -7,7 +7,8 @@
 #include "Modules/ModuleManager.h"
 
 /** Encryption Handler Component Module Interface */
-class FEncryptionHandlerComponentModuleInterface : public FPacketHandlerComponentModuleInterface
+class UE_DEPRECATED(5.3, "This component is not supported for encryption.")
+FEncryptionHandlerComponentModuleInterface : public FPacketHandlerComponentModuleInterface
 {
 	public:
 	/* Creates an instance of this component */
@@ -30,7 +31,8 @@ namespace EEncryptionHandler
 * to the remote side. Defaults are RSA for asymmetric encryption and XOR stream cipher
 * for symmetric encryption.
 */
-class ENCRYPTIONHANDLERCOMPONENT_API EncryptionHandlerComponent : public HandlerComponent
+class UE_DEPRECATED(5.3, "This component is not supported for encryption.")
+ENCRYPTIONHANDLERCOMPONENT_API EncryptionHandlerComponent : public HandlerComponent
 {
 public:
 	/* Initializes the default encryption handler components unless specified by the user */ 
diff --git a/Engine/Source/Runtime/PacketHandlers/EncryptionComponents/RSAKeyAESEncryption/Public/RSAKeyAESEncryptionHandlerComponent.h b/Engine/Source/Runtime/PacketHandlers/EncryptionComponents/RSAKeyAESEncryption/Public/RSAKeyAESEncryptionHandlerComponent.h
--- a/Engine/Source/Runtime/PacketHandlers/EncryptionComponents/RSAKeyAESEncryption/Public/RSAKeyAESEncryptionHandlerComponent.h
+++ b/Engine/Source/Runtime/PacketHandlers/EncryptionComponents/RSAKeyAESEncryption/Public/RSAKeyAESEncryptionHandlerComponent.h
@@ -19,8 +19,8 @@ THIRD_PARTY_INCLUDES_END
 /**
  * Module Interface
  */
-UE_DEPRECATED(5.0, "This Handler Is Now Deprecated");
-class FRSAKeyAESEncryptionModuleInterface : public FPacketHandlerComponentModuleInterface
+class UE_DEPRECATED(5.3, "This component is not supported for encryption.")
+FRSAKeyAESEncryptionModuleInterface : public FPacketHandlerComponentModuleInterface
 {
 public:
 	virtual TSharedPtr<HandlerComponent> CreateComponentInstance(FString& Options) override;
@@ -40,8 +40,8 @@ namespace ERSAKeyAESEncryptionHandler
 /**
  * Exchanges a key using RSA as the handshake, and implements symmetric encryption with that key using AES
  */
-UE_DEPRECATED(5.0, "This Handler Is Now Deprecated");
-class RSAKEYAESENCRYPTION_API RSAKeyAESEncryptionHandlerComponent : public HandlerComponent
+class UE_DEPRECATED(5.3, "This component is not supported for encryption.")
+RSAKEYAESENCRYPTION_API RSAKeyAESEncryptionHandlerComponent : public HandlerComponent
 {
 public:
 	/**
diff --git a/Engine/Source/Runtime/PacketHandlers/EncryptionComponents/SymmetricEncryption/BlockEncryption/AESBlockEncryptor/Public/AESBlockEncryptor.h b/Engine/Source/Runtime/PacketHandlers/EncryptionComponents/SymmetricEncryption/BlockEncryption/AESBlockEncryptor/Public/AESBlockEncryptor.h
--- a/Engine/Source/Runtime/PacketHandlers/EncryptionComponents/SymmetricEncryption/BlockEncryption/AESBlockEncryptor/Public/AESBlockEncryptor.h
+++ b/Engine/Source/Runtime/PacketHandlers/EncryptionComponents/SymmetricEncryption/BlockEncryption/AESBlockEncryptor/Public/AESBlockEncryptor.h
@@ -6,8 +6,8 @@
 #include "CryptoPP/5.6.5/include/aes.h"
 
 /* AES Block Encryptor Module Interface */
-UE_DEPRECATED(5.0, "This Encryptor Is Now Deprecated");
-class FAESBlockEncryptorModuleInterface : public FBlockEncryptorModuleInterface
+class UE_DEPRECATED(5.3, "This component is not supported for encryption.")
+FAESBlockEncryptorModuleInterface : public FBlockEncryptorModuleInterface
 {
 	virtual BlockEncryptor* CreateBlockEncryptorInstance() override;
 };
@@ -15,8 +15,8 @@ class FAESBlockEncryptorModuleInterface : public FBlockEncryptorModuleInterface
 /*
 * AES Block encryption
 */
-UE_DEPRECATED(5.0, "This Encryptor Is Now Deprecated");
-class AESBLOCKENCRYPTOR_API AESBlockEncryptor : public BlockEncryptor
+class UE_DEPRECATED(5.3, "This component is not supported for encryption.")
+AESBLOCKENCRYPTOR_API AESBlockEncryptor : public BlockEncryptor
 {
 public:
 	/* Initialized the encryptor */
diff --git a/Engine/Source/Runtime/PacketHandlers/EncryptionComponents/SymmetricEncryption/BlockEncryption/BlockEncryptionHandlerComponent/Public/BlockEncryptionHandlerComponent.h b/Engine/Source/Runtime/PacketHandlers/EncryptionComponents/SymmetricEncryption/BlockEncryption/BlockEncryptionHandlerComponent/Public/BlockEncryptionHandlerComponent.h
--- a/Engine/Source/Runtime/PacketHandlers/EncryptionComponents/SymmetricEncryption/BlockEncryption/BlockEncryptionHandlerComponent/Public/BlockEncryptionHandlerComponent.h
+++ b/Engine/Source/Runtime/PacketHandlers/EncryptionComponents/SymmetricEncryption/BlockEncryption/BlockEncryptionHandlerComponent/Public/BlockEncryptionHandlerComponent.h
@@ -9,8 +9,8 @@
 /*
 * Abstract class for block encryptors
 */
-UE_DEPRECATED(5.0, "This Encryptor Is Now Deprecated");
-class BLOCKENCRYPTIONHANDLERCOMPONENT_API BlockEncryptor
+class UE_DEPRECATED(5.3, "This component is not supported for encryption.")
+BLOCKENCRYPTIONHANDLERCOMPONENT_API BlockEncryptor
 {
 public:
 	/* Initialized the encryptor */
@@ -37,8 +37,8 @@ protected:
 };
 
 /* Block Encryptor Module Interface */
-UE_DEPRECATED(5.0, "This Encryptor Is Now Deprecated");
-class FBlockEncryptorModuleInterface : public IModuleInterface
+class UE_DEPRECATED(5.3, "This component is not supported for encryption.")
+FBlockEncryptorModuleInterface : public IModuleInterface
 {
 public:
 	virtual BlockEncryptor* CreateBlockEncryptorInstance() = 0;
@@ -47,8 +47,8 @@ public:
 /*
 * Symmetric block cipher handler component using AES as the cipher
 */
-UE_DEPRECATED(5.0, "This Encryptor Is Now Deprecated");
-class BLOCKENCRYPTIONHANDLERCOMPONENT_API BlockEncryptionHandlerComponent : public HandlerComponent
+class UE_DEPRECATED(5.3, "This component is not supported for encryption.")
+BLOCKENCRYPTIONHANDLERCOMPONENT_API BlockEncryptionHandlerComponent : public HandlerComponent
 {
 public:
 	/* Initializes default data, can provide the key size in bytes */
@@ -83,8 +83,8 @@ protected:
 };
 
 /* Block Encryption Handler Component Module Interface */
-UE_DEPRECATED(5.0, "This Encryptor Is Now Deprecated");
-class FBlockEncryptionHandlerComponentModuleInterface : public FPacketHandlerComponentModuleInterface
+class UE_DEPRECATED(5.3, "This component is not supported for encryption.")
+FBlockEncryptionHandlerComponentModuleInterface : public FPacketHandlerComponentModuleInterface
 {
 public:
 	/* Creates an instance of this component */
diff --git a/Engine/Source/Runtime/PacketHandlers/EncryptionComponents/SymmetricEncryption/BlockEncryption/BlowFishBlockEncryptor/Public/BlowFishBlockEncryptor.h b/Engine/Source/Runtime/PacketHandlers/EncryptionComponents/SymmetricEncryption/BlockEncryption/BlowFishBlockEncryptor/Public/BlowFishBlockEncryptor.h
--- a/Engine/Source/Runtime/PacketHandlers/EncryptionComponents/SymmetricEncryption/BlockEncryption/BlowFishBlockEncryptor/Public/BlowFishBlockEncryptor.h
+++ b/Engine/Source/Runtime/PacketHandlers/EncryptionComponents/SymmetricEncryption/BlockEncryption/BlowFishBlockEncryptor/Public/BlowFishBlockEncryptor.h
@@ -6,8 +6,8 @@
 #include "CryptoPP/5.6.5/include/blowfish.h"
 
 /* BlowFish Block Encryptor Module Interface */
-UE_DEPRECATED(5.0, "This Encryptor Is Now Deprecated");
-class FBlowFishBlockEncryptorModuleInterface : public FBlockEncryptorModuleInterface
+class UE_DEPRECATED(5.3, "This component is not supported for encryption.")
+FBlowFishBlockEncryptorModuleInterface : public FBlockEncryptorModuleInterface
 {
 	virtual BlockEncryptor* CreateBlockEncryptorInstance() override;
 };
@@ -16,8 +16,8 @@ class FBlowFishBlockEncryptorModuleInterface : public FBlockEncryptorModuleInter
 /*
 * BlowFish Block encryption
 */
-UE_DEPRECATED(5.0, "This Encryptor Is Now Deprecated");
-class BLOWFISHBLOCKENCRYPTOR_API BlowFishBlockEncryptor : public BlockEncryptor
+class UE_DEPRECATED(5.3, "This component is not supported for encryption.")
+BLOWFISHBLOCKENCRYPTOR_API BlowFishBlockEncryptor : public BlockEncryptor
 {
 public:
 	/* Initialized the encryptor */
diff --git a/Engine/Source/Runtime/PacketHandlers/EncryptionComponents/SymmetricEncryption/BlockEncryption/TwoFishBlockEncryptor/Public/TwoFishBlockEncryptor.h b/Engine/Source/Runtime/PacketHandlers/EncryptionComponents/SymmetricEncryption/BlockEncryption/TwoFishBlockEncryptor/Public/TwoFishBlockEncryptor.h
--- a/Engine/Source/Runtime/PacketHandlers/EncryptionComponents/SymmetricEncryption/BlockEncryption/TwoFishBlockEncryptor/Public/TwoFishBlockEncryptor.h
+++ b/Engine/Source/Runtime/PacketHandlers/EncryptionComponents/SymmetricEncryption/BlockEncryption/TwoFishBlockEncryptor/Public/TwoFishBlockEncryptor.h
@@ -6,8 +6,8 @@
 #include "CryptoPP/5.6.5/include/twofish.h"
 
 /* TwoFish Block Encryptor Module Interface */
-UE_DEPRECATED(5.0, "This Encryptor Is Now Deprecated");
-class FTwoFishBlockEncryptorModuleInterface : public FBlockEncryptorModuleInterface
+class UE_DEPRECATED(5.3, "This component is not supported for encryption.")
+FTwoFishBlockEncryptorModuleInterface : public FBlockEncryptorModuleInterface
 {
 	virtual BlockEncryptor* CreateBlockEncryptorInstance() override;
 };
@@ -15,8 +15,8 @@ class FTwoFishBlockEncryptorModuleInterface : public FBlockEncryptorModuleInterf
 /*
 * TwoFish Block encryption
 */
-UE_DEPRECATED(5.0, "This Encryptor Is Now Deprecated");
-class TWOFISHBLOCKENCRYPTOR_API TwoFishBlockEncryptor : public BlockEncryptor
+class UE_DEPRECATED(5.3, "This component is not supported for encryption.")
+TWOFISHBLOCKENCRYPTOR_API TwoFishBlockEncryptor : public BlockEncryptor
 {
 public:
 	/* Initialized the encryptor */
diff --git a/Engine/Source/Runtime/PacketHandlers/EncryptionComponents/SymmetricEncryption/BlockEncryption/XORBlockEncryptor/Public/XORBlockEncryptor.h b/Engine/Source/Runtime/PacketHandlers/EncryptionComponents/SymmetricEncryption/BlockEncryption/XORBlockEncryptor/Public/XORBlockEncryptor.h
--- a/Engine/Source/Runtime/PacketHandlers/EncryptionComponents/SymmetricEncryption/BlockEncryption/XORBlockEncryptor/Public/XORBlockEncryptor.h
+++ b/Engine/Source/Runtime/PacketHandlers/EncryptionComponents/SymmetricEncryption/BlockEncryption/XORBlockEncryptor/Public/XORBlockEncryptor.h
@@ -5,8 +5,8 @@
 #include "BlockEncryptionHandlerComponent.h"
 
 /* XOR Block Encryptor Module Interface */
-UE_DEPRECATED(5.0, "This Encryptor Is Now Deprecated");
-class FXORBlockEncryptorModuleInterface : public FBlockEncryptorModuleInterface
+class UE_DEPRECATED(5.3, "This component is not supported for encryption.")
+FXORBlockEncryptorModuleInterface : public FBlockEncryptorModuleInterface
 {
 	virtual BlockEncryptor* CreateBlockEncryptorInstance() override;
 };
@@ -14,8 +14,8 @@ class FXORBlockEncryptorModuleInterface : public FBlockEncryptorModuleInterface
 /*
 * XOR Block encryption
 */
-UE_DEPRECATED(5.0, "This Encryptor Is Now Deprecated");
-class XORBLOCKENCRYPTOR_API XORBlockEncryptor : public BlockEncryptor
+class UE_DEPRECATED(5.3, "This component is not supported for encryption.")
+XORBLOCKENCRYPTOR_API XORBlockEncryptor : public BlockEncryptor
 {
 public:
 	/* Initialized the encryptor */
diff --git a/Engine/Source/Runtime/PacketHandlers/EncryptionComponents/SymmetricEncryption/StreamEncryption/StreamEncryptionHandlerComponent/Public/StreamEncryptionHandlerComponent.h b/Engine/Source/Runtime/PacketHandlers/EncryptionComponents/SymmetricEncryption/StreamEncryption/StreamEncryptionHandlerComponent/Public/StreamEncryptionHandlerComponent.h
--- a/Engine/Source/Runtime/PacketHandlers/EncryptionComponents/SymmetricEncryption/StreamEncryption/StreamEncryptionHandlerComponent/Public/StreamEncryptionHandlerComponent.h
+++ b/Engine/Source/Runtime/PacketHandlers/EncryptionComponents/SymmetricEncryption/StreamEncryption/StreamEncryptionHandlerComponent/Public/StreamEncryptionHandlerComponent.h
@@ -4,8 +4,8 @@
 
 #include "PacketHandler.h"
 
-UE_DEPRECATED(5.0, "This Encryptor Is Now Deprecated");
-class STREAMENCRYPTIONHANDLERCOMPONENT_API StreamEncryptor
+class UE_DEPRECATED(5.3, "This component is not supported for encryption.")
+STREAMENCRYPTIONHANDLERCOMPONENT_API StreamEncryptor
 {
 public:
 	/* Initialized the encryptor */
@@ -26,8 +26,8 @@ protected:
 };
 
 /* Stream Encryptor Module Interface */
-UE_DEPRECATED(5.0, "This Encryptor Is Now Deprecated");
-class FStreamEncryptorModuleInterface : public IModuleInterface
+class UE_DEPRECATED(5.3, "This component is not supported for encryption.")
+FStreamEncryptorModuleInterface : public IModuleInterface
 {
 public:
 	virtual StreamEncryptor* CreateStreamEncryptorInstance() = 0;
@@ -36,8 +36,8 @@ public:
 /*
 * Symmetric block cipher handler component using AES as the cipher
 */
-UE_DEPRECATED(5.0, "This Encryptor Is Now Deprecated");
-class STREAMENCRYPTIONHANDLERCOMPONENT_API StreamEncryptionHandlerComponent : public HandlerComponent
+class UE_DEPRECATED(5.3, "This component is not supported for encryption.")
+STREAMENCRYPTIONHANDLERCOMPONENT_API StreamEncryptionHandlerComponent : public HandlerComponent
 {
 public:
 	/* Initializes default data, can provide the key size in bytes */
@@ -72,8 +72,8 @@ protected:
 };
 
 /* Stream Encryption Handler Component Module Interface */
-UE_DEPRECATED(5.0, "This Encryptor Is Now Deprecated");
-class FStreamEncryptionHandlerComponentModuleInterface : public FPacketHandlerComponentModuleInterface
+class UE_DEPRECATED(5.3, "This component is not supported for encryption.")
+FStreamEncryptionHandlerComponentModuleInterface : public FPacketHandlerComponentModuleInterface
 {
 public:
 	/* Creates an instance of this component */
diff --git a/Engine/Source/Runtime/PacketHandlers/EncryptionComponents/SymmetricEncryption/StreamEncryption/XORStreamEncryptor/Public/XORStreamEncryptor.h b/Engine/Source/Runtime/PacketHandlers/EncryptionComponents/SymmetricEncryption/StreamEncryption/XORStreamEncryptor/Public/XORStreamEncryptor.h
--- a/Engine/Source/Runtime/PacketHandlers/EncryptionComponents/SymmetricEncryption/StreamEncryption/XORStreamEncryptor/Public/XORStreamEncryptor.h
+++ b/Engine/Source/Runtime/PacketHandlers/EncryptionComponents/SymmetricEncryption/StreamEncryption/XORStreamEncryptor/Public/XORStreamEncryptor.h
@@ -3,8 +3,8 @@
 #include "StreamEncryptionHandlerComponent.h"
 
 /* XOR Stream Encryptor Module Interface */
-UE_DEPRECATED(5.0, "This Encryptor Is Now Deprecated");
-class FXORStreamEncryptorModuleInterface : public FStreamEncryptorModuleInterface
+class UE_DEPRECATED(5.3, "This component is not supported for encryption.")
+FXORStreamEncryptorModuleInterface : public FStreamEncryptorModuleInterface
 {
 public:
 	virtual StreamEncryptor* CreateStreamEncryptorInstance();
@@ -13,8 +13,8 @@ public:
 /*
 * XOR Block encryption
 */
-UE_DEPRECATED(5.0, "This Encryptor Is Now Deprecated");
-class XORSTREAMENCRYPTOR_API XORStreamEncryptor : public StreamEncryptor
+class UE_DEPRECATED(5.3, "This component is not supported for encryption.")
+XORSTREAMENCRYPTOR_API XORStreamEncryptor : public StreamEncryptor
 {
 public:
 	/* Initialized the encryptor */
diff --git a/Engine/Source/Runtime/PacketHandlers/PacketHandler/Private/PacketHandler.cpp b/Engine/Source/Runtime/PacketHandlers/PacketHandler/Private/PacketHandler.cpp
--- a/Engine/Source/Runtime/PacketHandlers/PacketHandler/Private/PacketHandler.cpp
+++ b/Engine/Source/Runtime/PacketHandlers/PacketHandler/Private/PacketHandler.cpp
@@ -215,9 +215,13 @@ void PacketHandler::Initialize(Handler::Mode InMode, uint32 InMaxPacketBits, boo
 
 	if (bEnableReliability && !ReliabilityComponent.IsValid())
 	{
+		UE_LOG(PacketHandlerLog, Warning, TEXT("bEnableReliability for PacketHandlerComponents is deprecated. For fully-reliable data, use reliable RPCs or a separate connection with a reliable protocol."));
+
+PRAGMA_DISABLE_DEPRECATION_WARNINGS
 		TSharedPtr<HandlerComponent> NewComponent = MakeShareable(new ReliabilityHandlerComponent);
 		ReliabilityComponent = StaticCastSharedPtr<ReliabilityHandlerComponent>(NewComponent);
 		AddHandler(NewComponent, true);
+PRAGMA_ENABLE_DEPRECATION_WARNINGS
 	}
 }
 
@@ -802,8 +806,10 @@ const ProcessedPacket PacketHandler::Outgoing_Internal(uint8* Packet, int32 Coun
 
 			if (!bConnectionless && ReliabilityComponent.IsValid() && OutgoingPacket.GetNumBits() > 0)
 			{
+PRAGMA_DISABLE_DEPRECATION_WARNINGS
 				// Let the reliability handler know about all processed packets, so it can record them for resending if needed
 				ReliabilityComponent->QueuePacketForResending(OutgoingPacket.GetData(), OutgoingPacket.GetNumBits(), Traits);
+PRAGMA_ENABLE_DEPRECATION_WARNINGS
 			}
 		}
 		// Buffer any packets being sent from game code until processors are initialized
@@ -969,8 +975,10 @@ void PacketHandler::SendHandlerPacket(HandlerComponent* InComponent, FBitWriter&
 
 			if (ReliabilityComponent.IsValid())
 			{
+PRAGMA_DISABLE_DEPRECATION_WARNINGS
 				// Let the reliability handler know about all processed packets, so it can record them for resending if needed
 				ReliabilityComponent->QueueHandlerPacketForResending(InComponent, Writer.GetData(), Writer.GetNumBits(), Traits);
+PRAGMA_ENABLE_DEPRECATION_WARNINGS
 			}
 
 			// Now finish off with a raw send (as we don't want to go through the PacketHandler chain again)
diff --git a/Engine/Source/Runtime/PacketHandlers/ReliabilityHandlerComponent/Private/ReliabilityHandlerComponent.cpp b/Engine/Source/Runtime/PacketHandlers/ReliabilityHandlerComponent/Private/ReliabilityHandlerComponent.cpp
--- a/Engine/Source/Runtime/PacketHandlers/ReliabilityHandlerComponent/Private/ReliabilityHandlerComponent.cpp
+++ b/Engine/Source/Runtime/PacketHandlers/ReliabilityHandlerComponent/Private/ReliabilityHandlerComponent.cpp
@@ -173,5 +173,7 @@ int32 ReliabilityHandlerComponent::GetReservedPacketBits() const
 // MODULE INTERFACE
 TSharedPtr<HandlerComponent> FReliabilityHandlerComponentModuleInterface::CreateComponentInstance(FString& Options)
 {
+PRAGMA_DISABLE_DEPRECATION_WARNINGS
 	return MakeShareable(new ReliabilityHandlerComponent);
+PRAGMA_ENABLE_DEPRECATION_WARNINGS
 }
diff --git a/Engine/Source/Runtime/PacketHandlers/ReliabilityHandlerComponent/Public/ReliabilityHandlerComponent.h b/Engine/Source/Runtime/PacketHandlers/ReliabilityHandlerComponent/Public/ReliabilityHandlerComponent.h
--- a/Engine/Source/Runtime/PacketHandlers/ReliabilityHandlerComponent/Public/ReliabilityHandlerComponent.h
+++ b/Engine/Source/Runtime/PacketHandlers/ReliabilityHandlerComponent/Public/ReliabilityHandlerComponent.h
@@ -15,8 +15,9 @@ class FString;
 struct FBitReader;
 struct FBitWriter;
 
-// Symmetric Stream cipher
-class RELIABILITYHANDLERCOMPONENT_API ReliabilityHandlerComponent : public HandlerComponent
+
+class UE_DEPRECATED(5.3, "This handler is not supported. For fully-reliable data, use reliable RPCs or a separate connection with a reliable protocol.")
+RELIABILITYHANDLERCOMPONENT_API ReliabilityHandlerComponent : public HandlerComponent
 {
 public:
 	/* Initializes default data */
