Enabled compile errors for unsafe floating point and integer casts in the networking modules along with fixes for existing errors.

diff --git a/Engine/Plugins/NetcodeUnitTest/NetcodeUnitTest/Source/NetcodeUnitTest/NetcodeUnitTest.Build.cs b/Engine/Plugins/NetcodeUnitTest/NetcodeUnitTest/Source/NetcodeUnitTest/NetcodeUnitTest.Build.cs
--- a/Engine/Plugins/NetcodeUnitTest/NetcodeUnitTest/Source/NetcodeUnitTest/NetcodeUnitTest.Build.cs
+++ b/Engine/Plugins/NetcodeUnitTest/NetcodeUnitTest/Source/NetcodeUnitTest/NetcodeUnitTest.Build.cs
@@ -42,6 +42,8 @@ namespace UnrealBuildTool.Rules
 					}
 				);
 			}
+
+			UnsafeTypeCastWarningLevel = WarningLevel.Error;
 		}
 	}
 }
diff --git a/Engine/Plugins/Runtime/ReplicationSystemTestPlugin/Source/ReplicationSystemTestPlugin.Build.cs b/Engine/Plugins/Runtime/ReplicationSystemTestPlugin/Source/ReplicationSystemTestPlugin.Build.cs
--- a/Engine/Plugins/Runtime/ReplicationSystemTestPlugin/Source/ReplicationSystemTestPlugin.Build.cs
+++ b/Engine/Plugins/Runtime/ReplicationSystemTestPlugin/Source/ReplicationSystemTestPlugin.Build.cs
@@ -55,6 +55,8 @@ namespace UnrealBuildTool.Rules
 				PrivateIncludePathModuleNames.Add("Catch2");
 			}
 
+			UnsafeTypeCastWarningLevel = WarningLevel.Error;
+
 			PrivateDefinitions.Add(String.Format("UE_NET_WITH_LOW_LEVEL_TESTS={0}", Target.ExplicitTestsTarget ? "1" : "0"));
 
 			SetupIrisSupport(Target);
diff --git a/Engine/Source/Runtime/NetworkFile/NetworkFile.Build.cs b/Engine/Source/Runtime/NetworkFile/NetworkFile.Build.cs
--- a/Engine/Source/Runtime/NetworkFile/NetworkFile.Build.cs
+++ b/Engine/Source/Runtime/NetworkFile/NetworkFile.Build.cs
@@ -17,5 +17,7 @@ public class NetworkFile : ModuleRules
 			});
 
 		PublicDefinitions.Add("ENABLE_HTTP_FOR_NETWORK_FILE=0");
+
+		UnsafeTypeCastWarningLevel = WarningLevel.Error;
 	}
 }
diff --git a/Engine/Source/Runtime/NetworkFile/Private/NetworkPlatformFile.cpp b/Engine/Source/Runtime/NetworkFile/Private/NetworkPlatformFile.cpp
--- a/Engine/Source/Runtime/NetworkFile/Private/NetworkPlatformFile.cpp
+++ b/Engine/Source/Runtime/NetworkFile/Private/NetworkPlatformFile.cpp
@@ -168,8 +168,15 @@ bool FNetworkPlatformFile::SendPayloadAndReceiveResponse(TArray<uint8>& In, TArr
 	}
 
 	TUniquePtr<FArchive> ResponsePayload = Response.ReadBody();
-	Out.SetNum(ResponsePayload->TotalSize());
-	ResponsePayload->Serialize(Out.GetData(), ResponsePayload->TotalSize());
+	
+	if (!IntFitsIn<int32, int64>(ResponsePayload->TotalSize()))
+	{
+		UE_LOG(LogCookOnTheFly, Warning, TEXT("Failed to parse 'CookOnTheFlyResponse' because the payload was too large"));
+		return false;
+	}
+	
+	Out.SetNum(static_cast<int32>(ResponsePayload->TotalSize()));
+	ResponsePayload->Serialize(Out.GetData(), Out.Num());
 	return true;
 }
 
@@ -182,8 +189,15 @@ void FNetworkPlatformFile::OnCookOnTheFlyMessage(const UE::Cook::FCookOnTheFlyMe
 	{
 		TUniquePtr<FArchive> PayloadReader = Message.ReadBody();
 		TArray<uint8> Payload;
-		Payload.SetNum(PayloadReader->TotalSize());
-		PayloadReader->Serialize(Payload.GetData(), PayloadReader->TotalSize());
+
+		if (!IntFitsIn<int32, int64>(PayloadReader->TotalSize()))
+		{
+			UE_LOG(LogCookOnTheFly, Warning, TEXT("Failed to parse 'CookOnTheFlyMessage' because the payload was too large"));
+			return;
+		}
+
+		Payload.SetNum(static_cast<int32>(PayloadReader->TotalSize()));
+		PayloadReader->Serialize(Payload.GetData(), Payload.Num());
 		PendingPayloads.Enqueue(MoveTemp(Payload));
 		NewPayloadEvent->Trigger();
 	}
diff --git a/Engine/Source/Runtime/NetworkFileSystem/NetworkFileSystem.Build.cs b/Engine/Source/Runtime/NetworkFileSystem/NetworkFileSystem.Build.cs
--- a/Engine/Source/Runtime/NetworkFileSystem/NetworkFileSystem.Build.cs
+++ b/Engine/Source/Runtime/NetworkFileSystem/NetworkFileSystem.Build.cs
@@ -38,6 +38,8 @@ namespace UnrealBuildTool.Rules
 			}
 
 			PrecompileForTargets = PrecompileTargetsType.Editor;
+
+			UnsafeTypeCastWarningLevel = WarningLevel.Error;
 		}
 	}
 }
diff --git a/Engine/Source/Runtime/NetworkFileSystem/Private/NetworkFileServerConnection.cpp b/Engine/Source/Runtime/NetworkFileSystem/Private/NetworkFileServerConnection.cpp
--- a/Engine/Source/Runtime/NetworkFileSystem/Private/NetworkFileServerConnection.cpp
+++ b/Engine/Source/Runtime/NetworkFileSystem/Private/NetworkFileServerConnection.cpp
@@ -1330,10 +1330,20 @@ bool FNetworkFileServerClientConnection::PackageFile( FString& Filename, FString
 		}
 		else
 		{
-			FileBytesSent += File->Size();
-			// read it
-			Contents.AddUninitialized(File->Size());
-			File->Read(Contents.GetData(), Contents.Num());
+			if (IntFitsIn<int32, int64>(File->Size()))
+			{
+				int32 FileSize32 = static_cast<int32>(File->Size());
+
+				FileBytesSent += FileSize32;
+				// read it
+				Contents.AddUninitialized(FileSize32);
+				File->Read(Contents.GetData(), Contents.Num());
+			}
+			else
+			{
+				UE_LOG(LogFileServer, Warning, TEXT("Unable to open %s because it is too large"), *Filename);
+				bRetVal = false;
+			}
 		}
 
 		// close it
diff --git a/Engine/Source/Runtime/Networking/Networking.Build.cs b/Engine/Source/Runtime/Networking/Networking.Build.cs
--- a/Engine/Source/Runtime/Networking/Networking.Build.cs
+++ b/Engine/Source/Runtime/Networking/Networking.Build.cs
@@ -12,6 +12,8 @@ namespace UnrealBuildTool.Rules
 					"Core",
 					"Sockets"
 				});
+
+			UnsafeTypeCastWarningLevel = WarningLevel.Error;
 		}
 	}
 }
diff --git a/Engine/Source/Runtime/Networking/Private/IPv4/IPv4Address.cpp b/Engine/Source/Runtime/Networking/Private/IPv4/IPv4Address.cpp
--- a/Engine/Source/Runtime/Networking/Private/IPv4/IPv4Address.cpp
+++ b/Engine/Source/Runtime/Networking/Private/IPv4/IPv4Address.cpp
@@ -29,10 +29,21 @@ bool FIPv4Address::Parse(const FString& AddressString, FIPv4Address& OutAddress)
 
 	if (AddressString.ParseIntoArray(Tokens, TEXT("."), false) == 4)
 	{
-		OutAddress.A = FCString::Atoi(*Tokens[0]);
-		OutAddress.B = FCString::Atoi(*Tokens[1]);
-		OutAddress.C = FCString::Atoi(*Tokens[2]);
-		OutAddress.D = FCString::Atoi(*Tokens[3]);
+		const int32 A = FCString::Atoi(*Tokens[0]);
+		const int32 B = FCString::Atoi(*Tokens[1]);
+		const int32 C = FCString::Atoi(*Tokens[2]);
+		const int32 D = FCString::Atoi(*Tokens[3]);
+
+		if (A < 0 || A > MAX_uint8 || B < 0 || B > MAX_uint8 ||
+			C < 0 || C > MAX_uint8 || D < 0 || D > MAX_uint8)
+		{
+			return false;
+		}
+
+		OutAddress.A = static_cast<uint8>(A);
+		OutAddress.B = static_cast<uint8>(B);
+		OutAddress.C = static_cast<uint8>(C);
+		OutAddress.D = static_cast<uint8>(D);
 
 		return true;
 	}
diff --git a/Engine/Source/Runtime/Networking/Private/IPv4/IPv4Endpoint.cpp b/Engine/Source/Runtime/Networking/Private/IPv4/IPv4Endpoint.cpp
--- a/Engine/Source/Runtime/Networking/Private/IPv4/IPv4Endpoint.cpp
+++ b/Engine/Source/Runtime/Networking/Private/IPv4/IPv4Endpoint.cpp
@@ -37,7 +37,14 @@ bool FIPv4Endpoint::Parse(const FString& EndpointString, FIPv4Endpoint& OutEndpo
 	{
 		if (FIPv4Address::Parse(Tokens[0], OutEndpoint.Address))
 		{
-			OutEndpoint.Port = FCString::Atoi(*Tokens[1]);
+			const int32 Port = FCString::Atoi(*Tokens[1]);
+
+			if (Port < 0 || Port > MAX_uint16)
+			{
+				return false;
+			}
+
+			OutEndpoint.Port = static_cast<uint16>(Port);
 
 			return true;
 		}
diff --git a/Engine/Source/Runtime/Networking/Private/IPv4/IPv4SubnetMask.cpp b/Engine/Source/Runtime/Networking/Private/IPv4/IPv4SubnetMask.cpp
--- a/Engine/Source/Runtime/Networking/Private/IPv4/IPv4SubnetMask.cpp
+++ b/Engine/Source/Runtime/Networking/Private/IPv4/IPv4SubnetMask.cpp
@@ -21,10 +21,21 @@ bool FIPv4SubnetMask::Parse(const FString& MaskString, FIPv4SubnetMask& OutMask)
 
 	if (MaskString.ParseIntoArray(Tokens, TEXT("."), false) == 4)
 	{
-		OutMask.A = FCString::Atoi(*Tokens[0]);
-		OutMask.B = FCString::Atoi(*Tokens[1]);
-		OutMask.C = FCString::Atoi(*Tokens[2]);
-		OutMask.D = FCString::Atoi(*Tokens[3]);
+		const int32 A = FCString::Atoi(*Tokens[0]);
+		const int32 B = FCString::Atoi(*Tokens[1]);
+		const int32 C = FCString::Atoi(*Tokens[2]);
+		const int32 D = FCString::Atoi(*Tokens[3]);
+
+		if (A < 0 || A > MAX_uint8 || B < 0 || B > MAX_uint8 ||
+			C < 0 || C > MAX_uint8 || D < 0 || D > MAX_uint8)
+		{
+			return false;
+		}
+
+		OutMask.A = static_cast<uint8>(A);
+		OutMask.B = static_cast<uint8>(B);
+		OutMask.C = static_cast<uint8>(C);
+		OutMask.D = static_cast<uint8>(D);
 
 		return true;
 	}
diff --git a/Engine/Source/Runtime/PacketHandlers/PacketHandler/PacketHandler.Build.cs b/Engine/Source/Runtime/PacketHandlers/PacketHandler/PacketHandler.Build.cs
--- a/Engine/Source/Runtime/PacketHandlers/PacketHandler/PacketHandler.Build.cs
+++ b/Engine/Source/Runtime/PacketHandlers/PacketHandler/PacketHandler.Build.cs
@@ -20,5 +20,7 @@ public class PacketHandler : ModuleRules
         );
 
         CircularlyReferencedDependentModules.Add("ReliabilityHandlerComponent");
-    }
+
+		UnsafeTypeCastWarningLevel = WarningLevel.Error;
+	}
 }
diff --git a/Engine/Source/Runtime/PacketHandlers/PacketHandler/Private/PacketHandler.cpp b/Engine/Source/Runtime/PacketHandlers/PacketHandler/Private/PacketHandler.cpp
--- a/Engine/Source/Runtime/PacketHandlers/PacketHandler/Private/PacketHandler.cpp
+++ b/Engine/Source/Runtime/PacketHandlers/PacketHandler/Private/PacketHandler.cpp
@@ -56,14 +56,14 @@ static FAutoConsoleVariableRef CVarNetPacketHandlerTimeguardLimit(
 
 // Lightweight time guard. Note: Threshold of 0 disables the timeguard
 #define NET_LIGHTWEIGHT_TIME_GUARD_BEGIN( Name, ThresholdMS ) \
-	float PREPROCESSOR_JOIN(__TimeGuard_ThresholdMS_, Name) = ThresholdMS; \
+	double PREPROCESSOR_JOIN(__TimeGuard_ThresholdMS_, Name) = ThresholdMS; \
 	uint64 PREPROCESSOR_JOIN(__TimeGuard_StartCycles_, Name) = ( ThresholdMS > 0.0f && GPacketHandlerTimeguardLimit > 0 ) ? FPlatformTime::Cycles64() : 0; \
 	GPacketHandlerDiscardTimeguardMeasurement = false;
 
 #define NET_LIGHTWEIGHT_TIME_GUARD_END( Name, NameStringCode ) \
 	if ( PREPROCESSOR_JOIN(__TimeGuard_ThresholdMS_, Name) > 0.0f && GPacketHandlerTimeguardLimit > 0 && !GPacketHandlerDiscardTimeguardMeasurement ) \
 	{\
-		float PREPROCESSOR_JOIN(__TimeGuard_MSElapsed_,Name) = FPlatformTime::ToMilliseconds64( FPlatformTime::Cycles64() - PREPROCESSOR_JOIN(__TimeGuard_StartCycles_,Name) ); \
+		double PREPROCESSOR_JOIN(__TimeGuard_MSElapsed_,Name) = FPlatformTime::ToMilliseconds64( FPlatformTime::Cycles64() - PREPROCESSOR_JOIN(__TimeGuard_StartCycles_,Name) ); \
 		if ( PREPROCESSOR_JOIN(__TimeGuard_MSElapsed_,Name) > PREPROCESSOR_JOIN(__TimeGuard_ThresholdMS_, Name) ) \
 		{ \
 			FString ReportName = NameStringCode; \
@@ -76,6 +76,15 @@ static FAutoConsoleVariableRef CVarNetPacketHandlerTimeguardLimit(
   #define NET_LIGHTWEIGHT_TIME_GUARD_END( Name, NameStringCode )
 #endif
 
+template<typename OutType, typename InType>
+OutType IntCastLog(InType In)
+{
+	bool bFitsIn = IntFitsIn<OutType, InType>(In);
+	ensureMsgf(bFitsIn, TEXT("PacketHandler: Loss of data caused by truncated cast"));
+	UE_CLOG(!bFitsIn, PacketHandlerLog, Warning, TEXT("PacketHandler: Loss of data caused by truncated cast"));
+	return static_cast<OutType>(In);
+}
+
 /**
  * BufferedPacket
  */
@@ -609,7 +618,7 @@ EIncomingResult PacketHandler::Incoming_Internal(FReceivedPacketView& PacketView
 				}
 				else if (ProcessedPacketReader.GetPosBits() == 0)
 				{
-					HandlerCRCs.Add({FCrc::MemCrc32(ProcessedPacketReader.GetData(), ProcessedPacketReader.GetNumBytes()), true, false});
+					HandlerCRCs.Add({FCrc::MemCrc32(ProcessedPacketReader.GetData(), IntCastLog<int32, int64>(ProcessedPacketReader.GetNumBytes())), true, false});
 				}
 				else
 				{
@@ -630,7 +639,7 @@ EIncomingResult PacketHandler::Incoming_Internal(FReceivedPacketView& PacketView
 					{
 						FHandlerCRC& CurCRC = HandlerCRCs[HandlerCRCs.Num() - 1];
 
-						CurCRC.CRC = FCrc::MemCrc32(ProcessedPacketReader.GetData(), ProcessedPacketReader.GetNumBytes());
+						CurCRC.CRC = FCrc::MemCrc32(ProcessedPacketReader.GetData(), IntCastLog<int32, int64>(ProcessedPacketReader.GetNumBytes()));
 						CurCRC.bHasAlignedCRC = true;
 					}
 #endif
@@ -662,7 +671,7 @@ EIncomingResult PacketHandler::Incoming_Internal(FReceivedPacketView& PacketView
 #if !UE_BUILD_SHIPPING
 				if (UNLIKELY(!!GPacketHandlerCRCDump))
 				{
-					NetConnectionCRC = FCrc::MemCrc32(IncomingPacket.GetData(), IncomingPacket.GetBytesLeft());
+					NetConnectionCRC = FCrc::MemCrc32(IncomingPacket.GetData(), IntCastLog<int32, int64>(IncomingPacket.GetBytesLeft()));
 				}
 #endif
 			}
@@ -790,7 +799,7 @@ const ProcessedPacket PacketHandler::Outgoing_Internal(uint8* Packet, int32 Coun
 					}
 					else
 					{
-						HandlerCRCs.Add({FCrc::MemCrc32(OutgoingPacket.GetData(), OutgoingPacket.GetNumBytes()), false});
+						HandlerCRCs.Add({FCrc::MemCrc32(OutgoingPacket.GetData(), IntCastLog<int32, int64>(OutgoingPacket.GetNumBytes())), false});
 					}
 				}
 #endif
@@ -808,7 +817,7 @@ const ProcessedPacket PacketHandler::Outgoing_Internal(uint8* Packet, int32 Coun
 			{
 PRAGMA_DISABLE_DEPRECATION_WARNINGS
 				// Let the reliability handler know about all processed packets, so it can record them for resending if needed
-				ReliabilityComponent->QueuePacketForResending(OutgoingPacket.GetData(), OutgoingPacket.GetNumBits(), Traits);
+				ReliabilityComponent->QueuePacketForResending(OutgoingPacket.GetData(), IntCastLog<int32, int64>(OutgoingPacket.GetNumBits()), Traits);
 PRAGMA_ENABLE_DEPRECATION_WARNINGS
 			}
 		}
@@ -892,7 +901,7 @@ void PacketHandler::ReplaceIncomingPacket(FBitReader& ReplacementPacket)
 	{
 		// @todo #JohnB: Make this directly adjust and write into IncomingPacket's buffer, instead of copying - very inefficient
 		TArray<uint8> TempPacketData;
-		TempPacketData.AddUninitialized(ReplacementPacket.GetBytesLeft());
+		TempPacketData.AddUninitialized(IntCastLog<int32, int64>(ReplacementPacket.GetBytesLeft()));
 		TempPacketData[TempPacketData.Num()-1] = 0;
 
 		int64 NewPacketSizeBits = ReplacementPacket.GetBitsLeft();
@@ -906,13 +915,13 @@ void PacketHandler::RealignPacket(FBitReader& Packet)
 {
 	if (Packet.GetPosBits() != 0)
 	{
-		uint32 BitsLeft = Packet.GetBitsLeft();
+		const int32 BitsLeft = IntCastLog<int32, int64>(Packet.GetBitsLeft());
 
 		if (BitsLeft > 0)
 		{
 			// @todo #JohnB: Based on above - when you optimize above, optimize this too
 			TArray<uint8> TempPacketData;
-			TempPacketData.AddUninitialized(Packet.GetBytesLeft());
+			TempPacketData.AddUninitialized(IntCastLog<int32, int64>(Packet.GetBytesLeft()));
 			TempPacketData[TempPacketData.Num()-1] = 0;
 
 			Packet.SerializeBits(TempPacketData.GetData(), BitsLeft);
@@ -972,12 +981,11 @@ void PacketHandler::SendHandlerPacket(HandlerComponent* InComponent, FBitWriter&
 			// Add a termination bit, the same as the UNetConnection code does, if appropriate
 			Writer.WriteBit(1);
 
-
 			if (ReliabilityComponent.IsValid())
 			{
 PRAGMA_DISABLE_DEPRECATION_WARNINGS
 				// Let the reliability handler know about all processed packets, so it can record them for resending if needed
-				ReliabilityComponent->QueueHandlerPacketForResending(InComponent, Writer.GetData(), Writer.GetNumBits(), Traits);
+				ReliabilityComponent->QueueHandlerPacketForResending(InComponent, Writer.GetData(), IntCastLog<int32, int64>(Writer.GetNumBits()), Traits);
 PRAGMA_ENABLE_DEPRECATION_WARNINGS
 			}
 
@@ -986,7 +994,7 @@ PRAGMA_ENABLE_DEPRECATION_WARNINGS
 
 			bRawSend = true;
 
-			LowLevelSendDel.ExecuteIfBound(Writer.GetData(), Writer.GetNumBits(), Traits);
+			LowLevelSendDel.ExecuteIfBound(Writer.GetData(), IntCastLog<int32, int64>(Writer.GetNumBits()), Traits);
 
 			bRawSend = bOldRawSend;
 		}
