Deprecate unnecessary PacketHandler delegates.

Plugin that uses these delegates is gone.

diff --git a/Engine/Source/Runtime/PacketHandlers/PacketHandler/Private/PacketHandler.cpp b/Engine/Source/Runtime/PacketHandlers/PacketHandler/Private/PacketHandler.cpp
--- a/Engine/Source/Runtime/PacketHandlers/PacketHandler/Private/PacketHandler.cpp
+++ b/Engine/Source/Runtime/PacketHandlers/PacketHandler/Private/PacketHandler.cpp
@@ -185,13 +185,18 @@ void PacketHandler::Initialize(UE::Handler::Mode InMode, uint32 InMaxPacketBits,
 		}
 
 		// Users of this delegate can add components to the list by name and if necessary reorder them
+		PRAGMA_DISABLE_DEPRECATION_WARNINGS
 		GetAddComponentByNameDelegate().ExecuteIfBound(Components);
+		PRAGMA_ENABLE_DEPRECATION_WARNINGS
 
 		for (const FString& CurComponent : Components)
 		{
 			AddHandler(CurComponent, true);
 		}
 
+		// all of code below till end of scope can go away once deprecated GetAddComponentDelegate method is removed
+		PRAGMA_DISABLE_DEPRECATION_WARNINGS
+
 		// Users of this delegate can supply constructed additional components to be added after the named components
 		TArray<TSharedPtr<HandlerComponent>> AdditionalComponents;
 		GetAddComponentDelegate().ExecuteIfBound(AdditionalComponents);
@@ -201,6 +206,7 @@ void PacketHandler::Initialize(UE::Handler::Mode InMode, uint32 InMaxPacketBits,
 			AddHandler(AdditionalComponent, true);
 		}
 
+		PRAGMA_ENABLE_DEPRECATION_WARNINGS
 	}
 
 	// Add encryption component, if configured.
diff --git a/Engine/Source/Runtime/PacketHandlers/PacketHandler/Public/PacketHandler.h b/Engine/Source/Runtime/PacketHandlers/PacketHandler/Public/PacketHandler.h
--- a/Engine/Source/Runtime/PacketHandlers/PacketHandler/Public/PacketHandler.h
+++ b/Engine/Source/Runtime/PacketHandlers/PacketHandler/Public/PacketHandler.h
@@ -647,7 +647,9 @@ public:
 	/** Mode of the handler, Client or Server */
 	UE::Handler::Mode Mode;
 
+	UE_DEPRECATED(5.6, "PacketHandler AddComponent delegates are deprecated, use [PacketHandlerComponents] section in .ini file")
 	static PACKETHANDLER_API FPacketHandlerAddComponentByNameDelegate& GetAddComponentByNameDelegate();
+	UE_DEPRECATED(5.6, "PacketHandler AddComponent delegates are deprecated, use [PacketHandlerComponents] section in .ini file")
 	static PACKETHANDLER_API FPacketHandlerAddComponentDelegate& GetAddComponentDelegate();
 
 private:
