Iris Parallelization - Added API for network HandlerComponents to indicate support for being run in parallel. Enabled NET_ANALYTICS_MULTITHREADING so that OodleNetworkAnalytics can run in parallel

[RN] Network HandlerComponents can indicate support for running in parallel per-connection. Enabled NET_ANALYTICS_MULTITHREADING by default.
#tests Built and tested on Dedicated Server with Cooked Data on Windows and Linux

#Synced-CL 48346153

diff --git a/Engine/Plugins/Compression/OodleNetwork/Source/Private/OodleNetworkHandlerComponent.cpp b/Engine/Plugins/Compression/OodleNetwork/Source/Private/OodleNetworkHandlerComponent.cpp
--- a/Engine/Plugins/Compression/OodleNetwork/Source/Private/OodleNetworkHandlerComponent.cpp
+++ b/Engine/Plugins/Compression/OodleNetwork/Source/Private/OodleNetworkHandlerComponent.cpp
@@ -1056,7 +1056,7 @@ void OodleNetworkHandlerComponent::Incoming(FIncomingPacketRef PacketRef)
 #if !UE_BUILD_SHIPPING
 							QUICK_SCOPE_CYCLE_COUNTER(STAT_Oodle_InCaptureTime);
 							CSV_SCOPED_TIMING_STAT_EXCLUSIVE(Oodle_InCaptureTime);
-							GPacketHandlerDiscardTimeguardMeasurement = true;
+							UE::Net::SetDiscardTimeguardMeasurement(true);
 #endif
 
 							InPacketLog->SerializePacket((void*)Packet.GetData(), DecompressedLength);
@@ -1126,7 +1126,7 @@ void OodleNetworkHandlerComponent::Incoming(FIncomingPacketRef PacketRef)
 #if !UE_BUILD_SHIPPING
 					QUICK_SCOPE_CYCLE_COUNTER(STAT_Oodle_InCaptureTime);
 					CSV_SCOPED_TIMING_STAT_EXCLUSIVE(Oodle_InCaptureTime);
-					GPacketHandlerDiscardTimeguardMeasurement = true;
+					UE::Net::SetDiscardTimeguardMeasurement(true);
 #endif
 
 					InPacketLog->SerializePacket((void*)Packet.GetData(), SizeOfPacket);
@@ -1153,7 +1153,7 @@ void OodleNetworkHandlerComponent::Outgoing(FBitWriter& Packet, FOutPacketTraits
 #if !UE_BUILD_SHIPPING
 				QUICK_SCOPE_CYCLE_COUNTER(STAT_Oodle_OutCaptureTime);
 				CSV_SCOPED_TIMING_STAT_EXCLUSIVE(Oodle_OutCaptureTime);
-				GPacketHandlerDiscardTimeguardMeasurement = true;
+				UE::Net::SetDiscardTimeguardMeasurement(true);
 #endif
 
 				OutPacketLog->SerializePacket((void*)Packet.GetData(), SizeOfPacket);
@@ -1477,6 +1477,16 @@ void OodleNetworkHandlerComponent::NotifyAnalyticsProvider()
 	bOodleNetworkAnalytics = NetAnalyticsData.IsValid();
 }
 
+bool OodleNetworkHandlerComponent::SupportsParallelConnectionOutgoing() const
+{
+	//Use of NetAnalytics means that we need Multithreading support for this component to support parallel connections
+#if NET_ANALYTICS_MULTITHREADING
+	return true;
+#else
+	return false;
+#endif //NET_ANALYTICS_MULTITHREADING
+}
+
 #if !UE_BUILD_SHIPPING
 /**
  * Exec Interface
diff --git a/Engine/Plugins/Compression/OodleNetwork/Source/Public/OodleNetworkHandlerComponent.h b/Engine/Plugins/Compression/OodleNetwork/Source/Public/OodleNetworkHandlerComponent.h
--- a/Engine/Plugins/Compression/OodleNetwork/Source/Public/OodleNetworkHandlerComponent.h
+++ b/Engine/Plugins/Compression/OodleNetwork/Source/Public/OodleNetworkHandlerComponent.h
@@ -362,6 +362,9 @@ public:
 	UE_API virtual int32 GetReservedPacketBits() const override;
 	UE_API virtual void NotifyAnalyticsProvider() override;
 
+	// Whether or not this HandlerComponent supports running its Outgoing logic on multiple separate connections simultaneously.
+	virtual bool SupportsParallelConnectionOutgoing() const override;
+
 
 protected:
 	/** Whether or not Oodle, and its additions to the packet protocol, are enabled */
diff --git a/Engine/Plugins/Online/OnlineSubsystemSteam/Source/Private/OnlineAuthHandlerSteam.cpp b/Engine/Plugins/Online/OnlineSubsystemSteam/Source/Private/OnlineAuthHandlerSteam.cpp
--- a/Engine/Plugins/Online/OnlineSubsystemSteam/Source/Private/OnlineAuthHandlerSteam.cpp
+++ b/Engine/Plugins/Online/OnlineSubsystemSteam/Source/Private/OnlineAuthHandlerSteam.cpp
@@ -447,6 +447,11 @@ void FSteamAuthHandlerComponent::SetComponentReady()
 	}
 }
 
+bool FSteamAuthHandlerComponent::SupportsParallelConnectionOutgoing() const
+{
+	return true;
+}
+
 /* Module handler */
 USteamAuthComponentModuleInterface::USteamAuthComponentModuleInterface(const FObjectInitializer& ObjectInitializer)
 	: Super(ObjectInitializer)
diff --git a/Engine/Plugins/Online/OnlineSubsystemSteam/Source/Private/OnlineAuthHandlerSteam.h b/Engine/Plugins/Online/OnlineSubsystemSteam/Source/Private/OnlineAuthHandlerSteam.h
--- a/Engine/Plugins/Online/OnlineSubsystemSteam/Source/Private/OnlineAuthHandlerSteam.h
+++ b/Engine/Plugins/Online/OnlineSubsystemSteam/Source/Private/OnlineAuthHandlerSteam.h
@@ -30,6 +30,8 @@ public:
 
 	virtual int32 GetReservedPacketBits() const override;
 
+	virtual bool SupportsParallelConnectionOutgoing() const override;
+
 protected:
 	enum class ESteamAuthHandlerState : uint8
 	{
diff --git a/Engine/Plugins/Online/OnlineSubsystemUtils/Source/OnlineSubsystemUtils/Private/IpConnection.cpp b/Engine/Plugins/Online/OnlineSubsystemUtils/Source/OnlineSubsystemUtils/Private/IpConnection.cpp
--- a/Engine/Plugins/Online/OnlineSubsystemUtils/Source/OnlineSubsystemUtils/Private/IpConnection.cpp
+++ b/Engine/Plugins/Online/OnlineSubsystemUtils/Source/OnlineSubsystemUtils/Private/IpConnection.cpp
@@ -21,6 +21,8 @@ Notes:
 #include "NetAddressResolution.h"
 #include "PacketHandler.h"
 
+#include "Iris/ReplicationSystem/ReplicationSystem.h"
+
 #include UE_INLINE_GENERATED_CPP_BY_NAME(IpConnection)
 
 
@@ -86,6 +88,10 @@ namespace UE::Net::Private
 		GNetRecreateSocketTimeoutThreshold,
 		TEXT("The amount of time, in seconds, without receiving a packet or alternatively without a send ack, before triggering socket recreation. ")
 		TEXT("(0.0 = off)"));
+
+	// Relates to the warning message which fires when trying to run HandlerComponents in parallel that do not support parallel execution
+	// This latch makes sure we only first the warning once across all threads
+	static std::atomic<bool> bHasLoggedHandlerParallelSupportWarning = false;
 }
 
 
@@ -176,7 +182,7 @@ void UIpConnection::Tick(float DeltaSeconds)
 {
 	using namespace UE::Net::Private;
 
-	if (CVarNetIpConnectionUseSendTasks.GetValueOnGameThread() != 0)
+	if (CVarNetIpConnectionUseSendTasks.GetValueOnAnyThread() != 0)
 	{
 		ISocketSubsystem* const SocketSubsystem = Driver->GetSocketSubsystem();
 		TArray<FSocketSendResult> ResultsCopy;
@@ -469,6 +475,23 @@ void UIpConnection::SendToRemote(uint8* DataToSend, int32 CountBits, FOutPacketT
 	// Process any packet modifiers, such as compression and encryption.
 	if (Handler.IsValid() && !Handler->GetRawSend())
 	{
+#if WITH_SERVER_CODE
+		//Flag a warning if we are trying to run non-parallel Handlers in a ParallelPhase
+		UE_AUTORTFM_OPEN 
+		{
+			const UReplicationSystem* ReplicationSystem = Driver->GetReplicationSystem();
+			if(ReplicationSystem && ReplicationSystem->IsInParallelPhase())
+			{
+				bool bExpected = false;
+				if(!Handler->CanRunHandlerComponentsOnParallelOutgoingConnections() && bHasLoggedHandlerParallelSupportWarning.compare_exchange_strong(bExpected, true, std::memory_order_relaxed))
+				{
+					UE_LOG(LogNet, Warning, TEXT("UIpConnection is trying to run non-parallel-compatible HandlerComponents in parallel. These HandlerComponents will be updated in a sequential order instead of simultaneously:"));
+					Handler->LogHandlerComponentList();
+				}
+			}
+		};
+#endif // WITH_SERVER_CODE
+
 		const ProcessedPacket ProcessedData = Handler->Outgoing(DataToSend, CountBits, Traits);
 
 		if (!ProcessedData.bError)
diff --git a/Engine/Plugins/Runtime/PacketHandlers/AESGCMHandlerComponent/Source/Private/AESGCMHandlerComponent.cpp b/Engine/Plugins/Runtime/PacketHandlers/AESGCMHandlerComponent/Source/Private/AESGCMHandlerComponent.cpp
--- a/Engine/Plugins/Runtime/PacketHandlers/AESGCMHandlerComponent/Source/Private/AESGCMHandlerComponent.cpp
+++ b/Engine/Plugins/Runtime/PacketHandlers/AESGCMHandlerComponent/Source/Private/AESGCMHandlerComponent.cpp
@@ -398,3 +398,8 @@ void FAESGCMHandlerComponent::CountBytes(FArchive& Ar) const
 	}
 	*/
 }
+
+bool FAESGCMHandlerComponent::SupportsParallelConnectionOutgoing() const
+{
+	return true;
+}
diff --git a/Engine/Plugins/Runtime/PacketHandlers/AESGCMHandlerComponent/Source/Public/AESGCMHandlerComponent.h b/Engine/Plugins/Runtime/PacketHandlers/AESGCMHandlerComponent/Source/Public/AESGCMHandlerComponent.h
--- a/Engine/Plugins/Runtime/PacketHandlers/AESGCMHandlerComponent/Source/Public/AESGCMHandlerComponent.h
+++ b/Engine/Plugins/Runtime/PacketHandlers/AESGCMHandlerComponent/Source/Public/AESGCMHandlerComponent.h
@@ -55,6 +55,9 @@ public:
 	AESGCMHANDLERCOMPONENT_API virtual int32 GetReservedPacketBits() const override;
 	AESGCMHANDLERCOMPONENT_API virtual void CountBytes(FArchive& Ar) const override;
 
+	// Whether or not this HandlerComponent supports running its Outgoing logic on multiple separate connections simultaneously.
+	virtual bool SupportsParallelConnectionOutgoing() const override;
+
 private:
 	TUniquePtr<FEncryptionContext> EncryptionContext;
 
diff --git a/Engine/Source/Runtime/Engine/Private/PacketHandlers/StatelessConnectHandlerComponent.cpp b/Engine/Source/Runtime/Engine/Private/PacketHandlers/StatelessConnectHandlerComponent.cpp
--- a/Engine/Source/Runtime/Engine/Private/PacketHandlers/StatelessConnectHandlerComponent.cpp
+++ b/Engine/Source/Runtime/Engine/Private/PacketHandlers/StatelessConnectHandlerComponent.cpp
@@ -950,6 +950,11 @@ void StatelessConnectHandlerComponent::SetHandshakeFailureCallback(FHandshakeFai
 	HandshakeFailureCallback = MoveTemp(InHandshakeFailureFunc);
 }
 
+bool StatelessConnectHandlerComponent::SupportsParallelConnectionOutgoing() const
+{
+	return true;
+}
+
 void StatelessConnectHandlerComponent::Incoming(FBitReader& Packet)
 {
 	using namespace UE::Net;
diff --git a/Engine/Source/Runtime/Engine/Public/PacketHandlers/StatelessConnectHandlerComponent.h b/Engine/Source/Runtime/Engine/Public/PacketHandlers/StatelessConnectHandlerComponent.h
--- a/Engine/Source/Runtime/Engine/Public/PacketHandlers/StatelessConnectHandlerComponent.h
+++ b/Engine/Source/Runtime/Engine/Public/PacketHandlers/StatelessConnectHandlerComponent.h
@@ -167,6 +167,14 @@ public:
 	 */
 	ENGINE_API void SetHandshakeFailureCallback(FHandshakeFailureFunc&& InHandshakeFailureFunc);
 
+	/**
+	 * Whether or not this HandlerComponent supports running its Outgoing logic on multiple separate connections simultaneously.
+	 * Any accesses on a per-connection basis are safe, but writing to Global/NetDriver properties are not
+	 *
+	 * @return	Whether or not the above is supported
+	 */
+	virtual bool SupportsParallelConnectionOutgoing() const override;
+
 private:
 	/**
 	 * Constructs and sends the initial handshake packet, from the client to the server
diff --git a/Engine/Source/Runtime/Net/Core/NetCore.Build.cs b/Engine/Source/Runtime/Net/Core/NetCore.Build.cs
--- a/Engine/Source/Runtime/Net/Core/NetCore.Build.cs
+++ b/Engine/Source/Runtime/Net/Core/NetCore.Build.cs
@@ -19,5 +19,12 @@ public class NetCore : ModuleRules
 		);
 
 		CppCompileWarningSettings.UnsafeTypeCastWarningLevel = WarningLevel.Error;
+
+		//Iris Parallel support requires Oodle to be threadsafe on the server
+		//So enable multithreading if not Client/Program and not disabled on the build command line 
+		if (!Target.bExcludeParallelIrisSupport && !(Target.Type == TargetType.Client || Target.Type == TargetType.Program))
+		{
+			PublicDefinitions.Add("NET_ANALYTICS_MULTITHREADING=1");
+		}
 	}
 }
diff --git a/Engine/Source/Runtime/Net/Core/Public/Net/Core/Analytics/NetAnalytics.h b/Engine/Source/Runtime/Net/Core/Public/Net/Core/Analytics/NetAnalytics.h
--- a/Engine/Source/Runtime/Net/Core/Public/Net/Core/Analytics/NetAnalytics.h
+++ b/Engine/Source/Runtime/Net/Core/Public/Net/Core/Analytics/NetAnalytics.h
@@ -33,7 +33,9 @@
  */
 
 // Includes
+#include "AutoRTFM.h"
 #include "Containers/Map.h"
+#include "Containers/Queue.h"
 #include "Containers/UnrealString.h"
 #include "CoreMinimal.h"
 #include "HAL/Platform.h"
@@ -51,7 +53,9 @@ struct FAnalyticsEventAttribute;
 // Defines
 
 /** Whether or not to enable multithreaded support (future proofing for netcode multithreading) - comes with a performance cost */
+#if !defined(NET_ANALYTICS_MULTITHREADING)
 #define NET_ANALYTICS_MULTITHREADING 0
+#endif
 
 
 inline constexpr const ESPMode NetAnalyticsThreadSafety = ESPMode::ThreadSafe;
@@ -230,7 +234,8 @@ template<class TDataStruct> struct TThreadedNetAnalyticsData : public FThreadedN
 	TDataStruct* GetLocalData()
 	{
 		/** Maps this threads net analytics variables, to the net analytics data handler which owns them - for fast lookup */
-		thread_local TMap<FThreadedNetAnalyticsData*, FLocalNetAnalyticsStruct*> LocalNetAnalyticsMap;
+		typedef TMap<FThreadedNetAnalyticsData*, FLocalNetAnalyticsStruct*> FLocalNetAnalyticsStructMap;
+		UE_AUTORTFM_DECLARE_THREAD_LOCAL_VAR(FLocalNetAnalyticsStructMap, LocalNetAnalyticsMap);
 
 		FLocalNetAnalyticsStruct** FoundVal = LocalNetAnalyticsMap.Find(this);
 
@@ -250,7 +255,8 @@ protected:
 		TDataStruct* ReturnVal = new TDataStruct();
 
 		/** Maps a weak net analytics data pointer, to its own raw pointer - to cleanup stale entries in LocalNetAnalyticsMap */
-		thread_local TMap<TNetAnalyticsDataWeakPtr<FThreadedNetAnalyticsData>, FThreadedNetAnalyticsData*> StaleNetAnalyticsTracking;
+		typedef TMap<TNetAnalyticsDataWeakPtr<FThreadedNetAnalyticsData>, FThreadedNetAnalyticsData*> FThreadedNetAnalyticsDataMap;
+		UE_AUTORTFM_DECLARE_THREAD_LOCAL_VAR(FThreadedNetAnalyticsDataMap, StaleNetAnalyticsTracking);
 
 		LocalNetAnalyticsMap.Add(this, ReturnVal);
 		StaleNetAnalyticsTracking.Add(StaticCastSharedRef<FThreadedNetAnalyticsData, FNetAnalyticsData, NetAnalyticsThreadSafety>(AsShared()), this);
@@ -258,7 +264,8 @@ protected:
 
 
 		/** Each thread tracks a counter signaling a new FNetAnalyticsAggregator/NetDriver instance, using this to trigger cleanup */
-		thread_local uint8 LastNetAnalyticsCounter = 0;
+		UE_AUTORTFM_DECLARE_THREAD_LOCAL_VAR(uint8, LastNetAnalyticsCounter);
+		LastNetAnalyticsCounter = 0;
 
 		uint8 CurNetAnalyticsCounter = GNetAnalyticsCounter;
 
diff --git a/Engine/Source/Runtime/PacketHandlers/PacketHandler/Private/PacketHandler.cpp b/Engine/Source/Runtime/PacketHandlers/PacketHandler/Private/PacketHandler.cpp
--- a/Engine/Source/Runtime/PacketHandlers/PacketHandler/Private/PacketHandler.cpp
+++ b/Engine/Source/Runtime/PacketHandlers/PacketHandler/Private/PacketHandler.cpp
@@ -16,7 +16,9 @@
 #include "PacketHandlerProfileConfig.h"
 
 #include "SocketSubsystem.h"
+#include "Misc/ScopeRWLock.h"
 #include "Misc/StringBuilder.h"
+#include "Misc/TransactionallySafeRWLock.h"
 
 // @todo #JohnB: There is quite a lot of inefficient copying of packet data going on.
 //					Redo the whole packet parsing/modification pipeline.
@@ -38,9 +40,10 @@ FAutoConsoleVariableRef CVarNetPacketHandlerCRCDump(
 	GPacketHandlerCRCDump,
 	TEXT("Enables or disables dumping of packet CRC's for every HandlerComponent, Incoming and Outgoing, for debugging."));
 
+static FTransactionallySafeRWLock PacketHandlerTimeguardLock;
 static int32 GPacketHandlerTimeguardLimit = 20;
 static float GPacketHandlerTimeguardThresholdMS = 0.0f;
-bool GPacketHandlerDiscardTimeguardMeasurement = false;
+static UE_AUTORTFM_DECLARE_THREAD_LOCAL_VAR(bool, GPacketHandlerDiscardTimeguardMeasurement);
 
 static FAutoConsoleVariableRef CVarNetPacketHandlerTimeguardThresholdMS(
 	TEXT("net.PacketHandlerTimeguardThresholdMS"),
@@ -54,21 +57,54 @@ static FAutoConsoleVariableRef CVarNetPacketHandlerTimeguardLimit(
 	ECVF_Default
 );
 
+namespace UE
+{
+	namespace Net
+	{
+		// GPacketHandlerDiscardTimeguardMeasurement is unique per-thread, but accessible from other
+		// compilation units, so it's exposed here.
+		void SetDiscardTimeguardMeasurement(const bool bShouldDiscard)
+		{
+			GPacketHandlerDiscardTimeguardMeasurement = bShouldDiscard;
+		}
+
+		bool GetDiscardTimeguardMeasurement()
+		{
+			return GPacketHandlerDiscardTimeguardMeasurement;
+		}
+
+		// GPacketHandlerTimeguardLimit can be accessed across threads to limit the total number of time guard warnings
+		// In cases where multiple threads go above the threshold simultaneously, it's possible to go above GPacketHandlerTimeguardLimit
+		// but after that frame no more warnings will be sent.
+		int32 GetTimeguardLimit()
+		{
+			UE::TReadScopeLock Lock(PacketHandlerTimeguardLock);
+			return GPacketHandlerTimeguardLimit;
+		}
+
+		void DecrementTimeguardLimit()
+		{
+			UE::TWriteScopeLock Lock(PacketHandlerTimeguardLock);
+			GPacketHandlerTimeguardLimit--;
+		}
+	}
+}
+
 // Lightweight time guard. Note: Threshold of 0 disables the timeguard
 #define NET_LIGHTWEIGHT_TIME_GUARD_BEGIN( Name, ThresholdMS ) \
 	double UE_JOIN(__TimeGuard_ThresholdMS_, Name) = ThresholdMS; \
-	uint64 UE_JOIN(__TimeGuard_StartCycles_, Name) = ( ThresholdMS > 0.0f && GPacketHandlerTimeguardLimit > 0 ) ? FPlatformTime::Cycles64() : 0; \
-	GPacketHandlerDiscardTimeguardMeasurement = false;
+	uint64 UE_JOIN(__TimeGuard_StartCycles_, Name) = ( ThresholdMS > 0.0f && UE::Net::GetTimeguardLimit() > 0 ) ? FPlatformTime::Cycles64() : 0; \
+	UE::Net::SetDiscardTimeguardMeasurement(false);
 
 #define NET_LIGHTWEIGHT_TIME_GUARD_END( Name, NameStringCode ) \
-	if ( UE_JOIN(__TimeGuard_ThresholdMS_, Name) > 0.0f && GPacketHandlerTimeguardLimit > 0 && !GPacketHandlerDiscardTimeguardMeasurement ) \
+	if ( UE_JOIN(__TimeGuard_ThresholdMS_, Name) > 0.0f && UE::Net::GetTimeguardLimit() > 0 && !UE::Net::GetDiscardTimeguardMeasurement() ) \
 	{\
 		double UE_JOIN(__TimeGuard_MSElapsed_,Name) = FPlatformTime::ToMilliseconds64( FPlatformTime::Cycles64() - UE_JOIN(__TimeGuard_StartCycles_,Name) ); \
 		if ( UE_JOIN(__TimeGuard_MSElapsed_,Name) > UE_JOIN(__TimeGuard_ThresholdMS_, Name) ) \
 		{ \
 			FString ReportName = NameStringCode; \
 			UE_LOG(PacketHandlerLog, Warning, TEXT("PacketHandler: %s - %s took %.2fms!"), TEXT(#Name), *ReportName, UE_JOIN(__TimeGuard_MSElapsed_,Name)); \
-			GPacketHandlerTimeguardLimit--; \
+			UE::Net::DecrementTimeguardLimit(); \
 		} \
 	}
 #else // UE_BUILD_SHIPPING
@@ -897,6 +933,34 @@ PRAGMA_ENABLE_DEPRECATION_WARNINGS
 	return ReturnVal;
 }
 
+bool PacketHandler::CanRunHandlerComponentsOnParallelOutgoingConnections()
+{
+	bool bCanRunInParallel = true;
+	for (int32 i = 0; i < HandlerComponents.Num(); ++i)
+	{
+		if (HandlerComponents[i])
+		{
+			HandlerComponent& CurComponent = *HandlerComponents[i];
+			bCanRunInParallel &= CurComponent.SupportsParallelConnectionOutgoing();
+		}
+	}
+
+	return bCanRunInParallel;
+}
+
+void PacketHandler::LogHandlerComponentList()
+{
+	for (int32 i = 0; i < HandlerComponents.Num(); ++i)
+	{
+		if (HandlerComponents[i])
+		{
+			HandlerComponent& CurComponent = *HandlerComponents[i];
+
+			UE_LOG(PacketHandlerLog, Log, TEXT("Handler: %s SupportsParallel: %s"), *CurComponent.GetName().ToString(), CurComponent.SupportsParallelConnectionOutgoing() ? TEXT("TRUE") : TEXT("FALSE"));
+		}
+	}
+}
+
 void PacketHandler::ReplaceIncomingPacket(FBitReader& ReplacementPacket)
 {
 	if (ReplacementPacket.GetPosBits() == 0 || ReplacementPacket.GetBitsLeft() == 0)
diff --git a/Engine/Source/Runtime/PacketHandlers/PacketHandler/Public/PacketHandler.h b/Engine/Source/Runtime/PacketHandlers/PacketHandler/Public/PacketHandler.h
--- a/Engine/Source/Runtime/PacketHandlers/PacketHandler/Public/PacketHandler.h
+++ b/Engine/Source/Runtime/PacketHandlers/PacketHandler/Public/PacketHandler.h
@@ -40,12 +40,15 @@ namespace UE
 	namespace Net
 	{
 		class FNetConnectionFaultRecoveryBase;
-	}
-}
 
 #if !UE_BUILD_SHIPPING
-extern PACKETHANDLER_API bool GPacketHandlerDiscardTimeguardMeasurement;
+		// bShouldDiscard=true will prevent the log warning from firing when calling NET_LIGHTWEIGHT_TIME_GUARD_END
+		// If called in a parallel thread, this will only apply to the current thread of execution
+		PACKETHANDLER_API void SetDiscardTimeguardMeasurement(const bool bShouldDiscard);
+		PACKETHANDLER_API bool GetDiscardTimeguardMeasurement();
 #endif
+	}
+}
 
 /**
  * Delegates
@@ -572,6 +575,15 @@ public:
 	 */
 	PACKETHANDLER_API int32 GetTotalReservedPacketBits();
 
+	/**
+	 * Whether or not all the HandlerComponents can support parallel Outgoing Connections
+	 */
+	PACKETHANDLER_API bool CanRunHandlerComponentsOnParallelOutgoingConnections();
+
+	/**
+	 * Output names of HandlerComponents to PacketHandlerLog
+	 */
+	PACKETHANDLER_API void LogHandlerComponentList();
 
 	/**
 	 * Sets whether or not outgoing packets should bypass this handler - used when raw packet sends are necessary
@@ -836,6 +848,16 @@ public:
 		return false;
 	}
 
+	/**
+	 * Whether or not this HandlerComponent supports running its Outgoing logic on multiple separate connections simultaneously.
+	 * Any accesses on a per-connection basis are safe, but writing to Global/NetDriver properties are not
+	 *
+	 * @return	Whether or not the above is supported
+	 */
+	virtual bool SupportsParallelConnectionOutgoing() const
+	{
+		return false;
+	}
 
 	/**
 	 * Initialization functionality should be placed here
